UnrealPak: Prevent non-deterministic padding data when compression is disabled.

diff --git a/Engine/Source/Developer/PakFileUtilities/Private/PakFileUtilities.cpp b/Engine/Source/Developer/PakFileUtilities/Private/PakFileUtilities.cpp
--- a/Engine/Source/Developer/PakFileUtilities/Private/PakFileUtilities.cpp
+++ b/Engine/Source/Developer/PakFileUtilities/Private/PakFileUtilities.cpp
@@ -282,8 +282,9 @@ bool PrepareCopyFileToPak(const FString& InMountPoint, const FPakInputPair& InFi
 			{
 				for(int64 FillIndex=FileSize; FillIndex < PaddedEncryptedFileSize && InFile.bNeedEncryption; ++FillIndex)
 				{
-					// Fill the trailing buffer with random bytes from file
-					InOutPersistentBuffer[FillIndex] = InOutPersistentBuffer[rand()%FileSize];
+					// Fill the trailing buffer with bytes from file. Note that this is now from a fixed location
+					// rather than a random one so that we produce deterministic results
+					InOutPersistentBuffer[FillIndex] = InOutPersistentBuffer[FillIndex % FileSize];
 				}
 
 				//Encrypt the buffer before writing it to disk
