Minor unrealpak refactor
 - Added a "-sign" parameter to enable signing for the given pak. The cached settings in the crypto.json file aren't used to enable signing any more.
 - Encapsulated all signing and encryption keys in a single keychain object which is passed around, rather than individual keys
 - Added some validation to UBT/UAT encryption config parsing so that if you enable any signing or encryption stuff but haven't generated a key, it'll warn and disable everything.

diff --git a/Engine/Source/Developer/PakFileUtilities/Private/PakFileUtilities.cpp b/Engine/Source/Developer/PakFileUtilities/Private/PakFileUtilities.cpp
--- a/Engine/Source/Developer/PakFileUtilities/Private/PakFileUtilities.cpp
+++ b/Engine/Source/Developer/PakFileUtilities/Private/PakFileUtilities.cpp
@@ -44,49 +44,47 @@ struct FKeyPair
 
 	bool IsValid() const
 	{
-		return !PrivateKey.Exponent.IsZero()
+		bool bAllKeysValid = !PrivateKey.Exponent.IsZero()
 			&& !PrivateKey.Modulus.IsZero()
 			&& !PublicKey.Exponent.IsZero()
 			&& !PublicKey.Modulus.IsZero();
-	}
-};
-
-bool TestKeys(FKeyPair& Pair)
-{
-	UE_LOG(LogPakFile, Display, TEXT("Testing signature keys."));
-
-	// Just some random values
-	static TEncryptionInt TestData[] =
-	{
-		11,
-		253,
-		128,
-		234,
-		56,
-		89,
-		34,
-		179,
-		29,
-		1024,
-		(int64)(MAX_int32),
-		(int64)(MAX_uint32)-1
-	};
-
-	for (int32 TestIndex = 0; TestIndex < ARRAY_COUNT(TestData); ++TestIndex)
-	{
-		TEncryptionInt EncryptedData = FEncryption::ModularPow(TestData[TestIndex], Pair.PrivateKey.Exponent, Pair.PrivateKey.Modulus);
-		TEncryptionInt DecryptedData = FEncryption::ModularPow(EncryptedData, Pair.PublicKey.Exponent, Pair.PublicKey.Modulus);
-		if (TestData[TestIndex] != DecryptedData)
-		{
-			UE_LOG(LogPakFile, Error, TEXT("Keys do not properly encrypt/decrypt data (failed test with %lld)"), TestData[TestIndex].ToInt());
-			return false;
+		
+		bool bIsValid = true;
+
+		if (bAllKeysValid)
+		{
+			// Just some random values
+			static TEncryptionInt TestData[] =
+			{
+				11,
+				253,
+				128,
+				234,
+				56,
+				89,
+				34,
+				179,
+				29,
+				1024,
+				(int64)(MAX_int32),
+				(int64)(MAX_uint32)-1
+			};
+
+			for (int32 TestIndex = 0; TestIndex < ARRAY_COUNT(TestData); ++TestIndex)
+			{
+				TEncryptionInt EncryptedData = FEncryption::ModularPow(TestData[TestIndex], PrivateKey.Exponent, PrivateKey.Modulus);
+				TEncryptionInt DecryptedData = FEncryption::ModularPow(EncryptedData, PublicKey.Exponent, PublicKey.Modulus);
+				if (TestData[TestIndex] != DecryptedData)
+				{
+					bIsValid = false;
+					break;
+				}
+			}
 		}
-	}
 
-	UE_LOG(LogPakFile, Display, TEXT("Signature keys check completed successfuly."));
-
-	return true;
-}
+		return bIsValid;
+	}
+};
 
 struct FNamedAESKey
 {
@@ -100,7 +98,12 @@ struct FNamedAESKey
 	}
 };
 
-typedef TMap<FGuid, FNamedAESKey> TKeyChain;
+struct FKeyChain
+{
+	FKeyPair SigningKey;
+	TMap<FGuid, FNamedAESKey> EncryptionKeys;
+	const FNamedAESKey* MasterEncryptionKey = nullptr;
+};
 
 class FMemoryCompressor;
 
@@ -239,6 +242,7 @@ struct FPakCommandLineParameters
 		, PatchSeekOptUseOrder(false)
 		, EncryptIndex(false)
 		, UseCustomCompressor(false)
+		, bSign(false)
 	{
 	}
 
@@ -255,6 +259,7 @@ struct FPakCommandLineParameters
 	bool EncryptIndex;
 	bool UseCustomCompressor;
 	FGuid EncryptionKeyGuid;
+	bool bSign;
 };
 
 struct FPakEntryPair
@@ -495,7 +500,7 @@ bool FCompressedFileBuffer::CompressFileToWorkingBuffer(const FPakInputPair& InF
 	return true;
 }
 
-bool PrepareCopyFileToPak(const FString& InMountPoint, const FPakInputPair& InFile, uint8*& InOutPersistentBuffer, int64& InOutBufferSize, FPakEntryPair& OutNewEntry, uint8*& OutDataToWrite, int64& OutSizeToWrite, const FNamedAESKey& InEncryptionKey)
+bool PrepareCopyFileToPak(const FString& InMountPoint, const FPakInputPair& InFile, uint8*& InOutPersistentBuffer, int64& InOutBufferSize, FPakEntryPair& OutNewEntry, uint8*& OutDataToWrite, int64& OutSizeToWrite, const FKeyChain& InKeyChain)
 {	
 	TUniquePtr<FArchive> FileHandle(IFileManager::Get().CreateFileReader(*InFile.Source));
 	bool bFileExists = FileHandle.IsValid();
@@ -508,7 +513,7 @@ bool PrepareCopyFileToPak(const FString& InMountPoint, const FPakInputPair& InFi
 		OutNewEntry.Info.Size = FileSize;
 		OutNewEntry.Info.UncompressedSize = FileSize;
 		OutNewEntry.Info.CompressionMethodIndex = 0;
-		OutNewEntry.Info.SetEncrypted( InFile.bNeedEncryption && InEncryptionKey.IsValid() );
+		OutNewEntry.Info.SetEncrypted( InFile.bNeedEncryption );
 		OutNewEntry.Info.SetDeleteRecord(false);
 
 		if (InOutBufferSize < PaddedEncryptedFileSize)
@@ -522,7 +527,7 @@ bool PrepareCopyFileToPak(const FString& InMountPoint, const FPakInputPair& InFi
 
 		{
 			OutSizeToWrite = FileSize;
-			if (InFile.bNeedEncryption && InEncryptionKey.IsValid())
+			if (InFile.bNeedEncryption)
 			{
 				for(int64 FillIndex=FileSize; FillIndex < PaddedEncryptedFileSize && InFile.bNeedEncryption; ++FillIndex)
 				{
@@ -532,7 +537,8 @@ bool PrepareCopyFileToPak(const FString& InMountPoint, const FPakInputPair& InFi
 				}
 
 				//Encrypt the buffer before writing it to disk
-				FAES::EncryptData(InOutPersistentBuffer, PaddedEncryptedFileSize, InEncryptionKey.Key);
+				check(InKeyChain.MasterEncryptionKey);
+				FAES::EncryptData(InOutPersistentBuffer, PaddedEncryptedFileSize, InKeyChain.MasterEncryptionKey->Key);
 				// Update the size to be written
 				OutSizeToWrite = PaddedEncryptedFileSize;
 				OutNewEntry.Info.SetEncrypted( true );
@@ -562,7 +568,7 @@ void FinalizeCopyCompressedFileToPak(FPakInfo& InPakInfo, const FCompressedFileB
 	}
 }
 
-bool PrepareCopyCompressedFileToPak(const FString& InMountPoint, FPakInfo& Info, const FPakInputPair& InFile, const FCompressedFileBuffer& CompressedFile, FPakEntryPair& OutNewEntry, uint8*& OutDataToWrite, int64& OutSizeToWrite, const FNamedAESKey& InEncryptionKey)
+bool PrepareCopyCompressedFileToPak(const FString& InMountPoint, FPakInfo& Info, const FPakInputPair& InFile, const FCompressedFileBuffer& CompressedFile, FPakEntryPair& OutNewEntry, uint8*& OutDataToWrite, int64& OutSizeToWrite, const FKeyChain& InKeyChain)
 {
 	if (CompressedFile.TotalCompressedSize == 0)
 	{
@@ -572,9 +578,10 @@ bool PrepareCopyCompressedFileToPak(const FString& InMountPoint, FPakInfo& Info,
 	OutNewEntry.Info.CompressionMethodIndex = Info.GetCompressionMethodIndex(CompressedFile.FileCompressionMethod);
 	OutNewEntry.Info.CompressionBlocks.AddZeroed(CompressedFile.CompressedBlocks.Num());
 
-	if (InFile.bNeedEncryption && InEncryptionKey.IsValid())
+	if (InFile.bNeedEncryption)
 	{
-		FAES::EncryptData(CompressedFile.CompressedBuffer.Get(), CompressedFile.TotalCompressedSize, InEncryptionKey.Key);
+		check(InKeyChain.MasterEncryptionKey);
+		FAES::EncryptData(CompressedFile.CompressedBuffer.Get(), CompressedFile.TotalCompressedSize, InKeyChain.MasterEncryptionKey->Key);
 	}
 
 	//Hash the final buffer thats written
@@ -591,7 +598,7 @@ bool PrepareCopyCompressedFileToPak(const FString& InMountPoint, FPakInfo& Info,
 	//	Write the header, then the data
 	OutNewEntry.Filename = InFile.Dest.Mid(InMountPoint.Len());
 	OutNewEntry.Info.Offset = 0; // Don't serialize offsets here.
-	OutNewEntry.Info.SetEncrypted( InFile.bNeedEncryption && InEncryptionKey.IsValid() );
+	OutNewEntry.Info.SetEncrypted( InFile.bNeedEncryption );
 	OutNewEntry.Info.SetDeleteRecord(false);
 	OutSizeToWrite = CompressedFile.TotalCompressedSize;
 	OutDataToWrite = CompressedFile.CompressedBuffer.Get();
@@ -730,10 +737,9 @@ void ProcessCommandLine(const TCHAR* CmdLine, const TArray<FString>& NonOptionAr
 		CmdLineParameters.EncryptIndex = true;
 	}
 
-	FString EncryptionKeyGuid;
-	if (FParse::Value(CmdLine, TEXT("EncryptionKeyOverrideGuid="), EncryptionKeyGuid))
+	if (FParse::Param(CmdLine, TEXT("sign")))
 	{
-		FGuid::Parse(EncryptionKeyGuid, CmdLineParameters.EncryptionKeyGuid);
+		CmdLineParameters.bSign = true;
 	}
 
 	FString DesiredCompressionFormats;
@@ -1066,7 +1072,7 @@ void CollectFilesToAdd(TArray<FPakInputPair>& OutFilesToAdd, const TArray<FPakIn
 	UE_LOG(LogPakFile, Display, TEXT("Collected %d files in %.2lfs."), OutFilesToAdd.Num(), FPlatformTime::Seconds() - StartTime);
 }
 
-bool BufferedCopyFile(FArchive& Dest, FArchive& Source, const FPakFile& PakFile, const FPakEntry& Entry, void* Buffer, int64 BufferSize, const TKeyChain& KeyChain)
+bool BufferedCopyFile(FArchive& Dest, FArchive& Source, const FPakFile& PakFile, const FPakEntry& Entry, void* Buffer, int64 BufferSize, const FKeyChain& InKeyChain)
 {	
 	// Align down
 	BufferSize = BufferSize & ~(FAES::AESBlockSize-1);
@@ -1080,7 +1086,7 @@ bool BufferedCopyFile(FArchive& Dest, FArchive& Source, const FPakFile& PakFile,
 		Source.Serialize(Buffer,SizeToRead);
 		if (Entry.IsEncrypted())
 		{
-			const FNamedAESKey* Key = KeyChain.Find(PakFile.GetInfo().EncryptionKeyGuid);
+			const FNamedAESKey* Key = InKeyChain.MasterEncryptionKey;
 			check(Key);
 			FAES::DecryptData((uint8*)Buffer, SizeToRead, Key->Key);
 		}
@@ -1090,7 +1096,7 @@ bool BufferedCopyFile(FArchive& Dest, FArchive& Source, const FPakFile& PakFile,
 	return true;
 }
 
-bool UncompressCopyFile(FArchive& Dest, FArchive& Source, const FPakEntry& Entry, uint8*& PersistentBuffer, int64& BufferSize, const TKeyChain& KeyChain, const FPakFile& PakFile)
+bool UncompressCopyFile(FArchive& Dest, FArchive& Source, const FPakEntry& Entry, uint8*& PersistentBuffer, int64& BufferSize, const FKeyChain& InKeyChain, const FPakFile& PakFile)
 {
 	if (Entry.UncompressedSize == 0)
 	{
@@ -1126,7 +1132,7 @@ bool UncompressCopyFile(FArchive& Dest, FArchive& Source, const FPakEntry& Entry
 
 		if (Entry.IsEncrypted())
 		{
-			const FNamedAESKey* Key = KeyChain.Find(PakFile.GetInfo().EncryptionKeyGuid);
+			const FNamedAESKey* Key = InKeyChain.MasterEncryptionKey;
 			check(Key);
 			FAES::DecryptData(PersistentBuffer, SizeToRead, Key->Key);
 		}
@@ -1157,7 +1163,7 @@ TEncryptionInt ParseEncryptionIntFromJson(TSharedPtr<FJsonObject> InObj, const T
 	}
 }
 
-void PrepareEncryptionAndSigningKeysFromCryptoKeyCache(const FString& InFilename, FKeyPair& OutSigningKey, TKeyChain& OutKeyChain)
+void LoadKeyChainFromFile(const FString& InFilename, FKeyChain& OutCryptoSettings)
 {
 	FArchive* File = IFileManager::Get().CreateFileReader(*InFilename);
 	UE_CLOG(File == nullptr, LogPakFile, Fatal, TEXT("Specified crypto keys cache '%s' does not exist!"), *InFilename);
@@ -1165,45 +1171,36 @@ void PrepareEncryptionAndSigningKeysFromCryptoKeyCache(const FString& InFilename
 	TSharedRef<TJsonReader<char>> Reader = TJsonReaderFactory<char>::Create(File);
 	if (FJsonSerializer::Deserialize(Reader, RootObject))
 	{
-		const bool bDataCryptoRequired = RootObject->GetBoolField("bDataCryptoRequired");
-
-		if (bDataCryptoRequired)
+		const TSharedPtr<FJsonObject>* EncryptionKeyObject;
+		if (RootObject->TryGetObjectField(TEXT("EncryptionKey"), EncryptionKeyObject))
 		{
-			const TSharedPtr<FJsonObject>* EncryptionKeyObject;
-			if (RootObject->TryGetObjectField(TEXT("EncryptionKey"), EncryptionKeyObject))
+			FString EncryptionKeyBase64;
+			if ((*EncryptionKeyObject)->TryGetStringField(TEXT("Key"), EncryptionKeyBase64))
 			{
-				FString EncryptionKeyBase64;
-				if ((*EncryptionKeyObject)->TryGetStringField(TEXT("Key"), EncryptionKeyBase64))
+				if (EncryptionKeyBase64.Len() > 0)
 				{
-					if (EncryptionKeyBase64.Len() > 0)
-					{
-						TArray<uint8> Key;
-						FBase64::Decode(EncryptionKeyBase64, Key);
-						check(Key.Num() == sizeof(FAES::FAESKey::Key));
-						FNamedAESKey NewKey;
-						NewKey.Name = TEXT("Default");
-						NewKey.Guid = FGuid();
-						FMemory::Memcpy(NewKey.Key.Key, &Key[0], sizeof(FAES::FAESKey::Key));
-						OutKeyChain.Add(NewKey.Guid, NewKey);
-					}
+					TArray<uint8> Key;
+					FBase64::Decode(EncryptionKeyBase64, Key);
+					check(Key.Num() == sizeof(FAES::FAESKey::Key));
+					FNamedAESKey NewKey;
+					NewKey.Name = TEXT("Default");
+					NewKey.Guid = FGuid();
+					FMemory::Memcpy(NewKey.Key.Key, &Key[0], sizeof(FAES::FAESKey::Key));
+					OutCryptoSettings.EncryptionKeys.Add(NewKey.Guid, NewKey);
 				}
 			}
+		}
 
-			bool bEnablePakSigning = false;
-			if (RootObject->TryGetBoolField(TEXT("bEnablePakSigning"), bEnablePakSigning))
-			{
-				const TSharedPtr<FJsonObject>* SigningKey = nullptr;
-				if (bEnablePakSigning && RootObject->TryGetObjectField(TEXT("SigningKey"), SigningKey))
-				{
-					TSharedPtr<FJsonObject> PublicKey = (*SigningKey)->GetObjectField(TEXT("PublicKey"));
-					TSharedPtr<FJsonObject> PrivateKey = (*SigningKey)->GetObjectField(TEXT("PrivateKey"));
-					OutSigningKey.PublicKey.Exponent = ParseEncryptionIntFromJson(PublicKey, TEXT("Exponent"));
-					OutSigningKey.PublicKey.Modulus = ParseEncryptionIntFromJson(PublicKey, TEXT("Modulus"));
-					OutSigningKey.PrivateKey.Exponent = ParseEncryptionIntFromJson(PrivateKey, TEXT("Exponent"));
-					OutSigningKey.PrivateKey.Modulus = ParseEncryptionIntFromJson(PrivateKey, TEXT("Modulus"));
-					check(OutSigningKey.PublicKey.Modulus == OutSigningKey.PrivateKey.Modulus);
-				}
-			}
+		const TSharedPtr<FJsonObject>* SigningKey = nullptr;
+		if (RootObject->TryGetObjectField(TEXT("SigningKey"), SigningKey))
+		{
+			TSharedPtr<FJsonObject> PublicKey = (*SigningKey)->GetObjectField(TEXT("PublicKey"));
+			TSharedPtr<FJsonObject> PrivateKey = (*SigningKey)->GetObjectField(TEXT("PrivateKey"));
+			OutCryptoSettings.SigningKey.PublicKey.Exponent = ParseEncryptionIntFromJson(PublicKey, TEXT("Exponent"));
+			OutCryptoSettings.SigningKey.PublicKey.Modulus = ParseEncryptionIntFromJson(PublicKey, TEXT("Modulus"));
+			OutCryptoSettings.SigningKey.PrivateKey.Exponent = ParseEncryptionIntFromJson(PrivateKey, TEXT("Exponent"));
+			OutCryptoSettings.SigningKey.PrivateKey.Modulus = ParseEncryptionIntFromJson(PrivateKey, TEXT("Modulus"));
+			check(OutCryptoSettings.SigningKey.PublicKey.Modulus == OutCryptoSettings.SigningKey.PrivateKey.Modulus);
 		}
 
 		const TArray<TSharedPtr<FJsonValue>>* SecondaryEncryptionKeyArray = nullptr;
@@ -1222,28 +1219,28 @@ void PrepareEncryptionAndSigningKeysFromCryptoKeyCache(const FString& InFilename
 				check(Key.Num() == sizeof(FAES::FAESKey::Key));
 				FMemory::Memcpy(NewKey.Key.Key, &Key[0], sizeof(FAES::FAESKey::Key));
 
-				check(!OutKeyChain.Contains(NewKey.Guid) || OutKeyChain[NewKey.Guid].Key == NewKey.Key);
-				OutKeyChain.Add(NewKey.Guid, NewKey);
+				check(!OutCryptoSettings.EncryptionKeys.Contains(NewKey.Guid) || OutCryptoSettings.EncryptionKeys[NewKey.Guid].Key == NewKey.Key);
+				OutCryptoSettings.EncryptionKeys.Add(NewKey.Guid, NewKey);
 			}
 		}
 	}
 	delete File;
 }
 
-void PrepareEncryptionAndSigningKeys(const TCHAR* CmdLine, FKeyPair& OutSigningKey, TKeyChain& OutKeyChain)
+void LoadKeyChain(const TCHAR* CmdLine, FKeyChain& OutCryptoSettings)
 {
-	OutSigningKey.PrivateKey.Exponent.Zero();
-	OutSigningKey.PrivateKey.Modulus.Zero();
-	OutSigningKey.PublicKey.Exponent.Zero();
-	OutSigningKey.PublicKey.Modulus.Zero();
-	OutKeyChain.Empty();
+	OutCryptoSettings.SigningKey.PrivateKey.Exponent.Zero();
+	OutCryptoSettings.SigningKey.PrivateKey.Modulus.Zero();
+	OutCryptoSettings.SigningKey.PublicKey.Exponent.Zero();
+	OutCryptoSettings.SigningKey.PublicKey.Modulus.Zero();
+	OutCryptoSettings.EncryptionKeys.Empty();
 
 	// First, try and parse the keys from a supplied crypto key cache file
 	FString CryptoKeysCacheFilename;
 	if (FParse::Value(CmdLine, TEXT("cryptokeys="), CryptoKeysCacheFilename))
 	{
 		UE_LOG(LogPakFile, Display, TEXT("Parsing crypto keys from a crypto key cache file"));
-		PrepareEncryptionAndSigningKeysFromCryptoKeyCache(CryptoKeysCacheFilename, OutSigningKey, OutKeyChain);
+		LoadKeyChainFromFile(CryptoKeysCacheFilename, OutCryptoSettings);
 	}
 	else if (FParse::Param(CmdLine, TEXT("encryptionini")))
 	{
@@ -1296,10 +1293,10 @@ void PrepareEncryptionAndSigningKeys(const TCHAR* CmdLine, FKeyPair& OutSigningK
 					FBase64::Decode(PrivateExpBase64, PrivateExp);
 					FBase64::Decode(ModulusBase64, Modulus);
 
-					OutSigningKey.PrivateKey.Exponent = TEncryptionInt((uint32*)&PrivateExp[0]);
-					OutSigningKey.PrivateKey.Modulus = TEncryptionInt((uint32*)&Modulus[0]);
-					OutSigningKey.PublicKey.Exponent = TEncryptionInt((uint32*)&PublicExp[0]);
-					OutSigningKey.PublicKey.Modulus = OutSigningKey.PrivateKey.Modulus;
+					OutCryptoSettings.SigningKey.PrivateKey.Exponent = TEncryptionInt((uint32*)&PrivateExp[0]);
+					OutCryptoSettings.SigningKey.PrivateKey.Modulus = TEncryptionInt((uint32*)&Modulus[0]);
+					OutCryptoSettings.SigningKey.PublicKey.Exponent = TEncryptionInt((uint32*)&PublicExp[0]);
+					OutCryptoSettings.SigningKey.PublicKey.Modulus = OutCryptoSettings.SigningKey.PrivateKey.Modulus;
 
 					UE_LOG(LogPakFile, Display, TEXT("Parsed signature keys from config files."));
 				}
@@ -1318,7 +1315,7 @@ void PrepareEncryptionAndSigningKeys(const TCHAR* CmdLine, FKeyPair& OutSigningK
 						NewKey.Name = TEXT("Default");
 						NewKey.Guid = FGuid();
 						FMemory::Memcpy(NewKey.Key.Key, &Key[0], sizeof(FAES::FAESKey::Key));
-						OutKeyChain.Add(NewKey.Guid, NewKey);
+						OutCryptoSettings.EncryptionKeys.Add(NewKey.Guid, NewKey);
 						UE_LOG(LogPakFile, Display, TEXT("Parsed AES encryption key from config files."));
 					}
 				}
@@ -1340,10 +1337,10 @@ void PrepareEncryptionAndSigningKeys(const TCHAR* CmdLine, FKeyPair& OutSigningK
 					ConfigFile.GetString(SectionName, TEXT("rsa.privateexp"), RSAPrivateExp);
 					ConfigFile.GetString(SectionName, TEXT("rsa.modulus"), RSAModulus);
 
-					OutSigningKey.PrivateKey.Exponent.Parse(RSAPrivateExp);
-					OutSigningKey.PrivateKey.Modulus.Parse(RSAModulus);
-					OutSigningKey.PublicKey.Exponent.Parse(RSAPublicExp);
-					OutSigningKey.PublicKey.Modulus = OutSigningKey.PrivateKey.Modulus;
+					OutCryptoSettings.SigningKey.PrivateKey.Exponent.Parse(RSAPrivateExp);
+					OutCryptoSettings.SigningKey.PrivateKey.Modulus.Parse(RSAModulus);
+					OutCryptoSettings.SigningKey.PublicKey.Exponent.Parse(RSAPublicExp);
+					OutCryptoSettings.SigningKey.PublicKey.Modulus = OutCryptoSettings.SigningKey.PrivateKey.Modulus;
 
 					UE_LOG(LogPakFile, Display, TEXT("Parsed signature keys from config files."));
 				}
@@ -1361,7 +1358,7 @@ void PrepareEncryptionAndSigningKeys(const TCHAR* CmdLine, FKeyPair& OutSigningK
 						{
 							NewKey.Key.Key[Index] = (uint8)EncryptionKeyString[Index];
 						}
-						OutKeyChain.Add(NewKey.Guid, NewKey);
+						OutCryptoSettings.EncryptionKeys.Add(NewKey.Guid, NewKey);
 						UE_LOG(LogPakFile, Display, TEXT("Parsed AES encryption key from config files."));
 					}
 				}
@@ -1402,30 +1399,34 @@ void PrepareEncryptionAndSigningKeys(const TCHAR* CmdLine, FKeyPair& OutSigningK
 			ANSICHAR* AsAnsi = TCHAR_TO_ANSI(*EncryptionKeyString);
 			check(TCString<ANSICHAR>::Strlen(AsAnsi) == RequiredKeyLength);
 			FMemory::Memcpy(NewKey.Key.Key, AsAnsi, RequiredKeyLength);
-			OutKeyChain.Add(NewKey.Guid, NewKey);
+			OutCryptoSettings.EncryptionKeys.Add(NewKey.Guid, NewKey);
 			UE_LOG(LogPakFile, Display, TEXT("Parsed AES encryption key from command line."));
 		}
 	}
 
-	if (OutSigningKey.IsValid())
+	if (!OutCryptoSettings.SigningKey.IsValid())
 	{
-		if (!TestKeys(OutSigningKey))
-		{
-			UE_LOG(LogPakFile, Fatal, TEXT("Pak signing keys are invalid"));
-			OutSigningKey.PrivateKey.Exponent.Zero();
-		}
+		UE_LOG(LogPakFile, Fatal, TEXT("Pak signing keys are invalid"));
+	}
+
+	FString EncryptionKeyOverrideGuidString;
+	FGuid EncryptionKeyOverrideGuid;
+	if (FParse::Value(CmdLine, TEXT("EncryptionKeyOverrideGuid="), EncryptionKeyOverrideGuidString))
+	{
+		FGuid::Parse(EncryptionKeyOverrideGuidString, EncryptionKeyOverrideGuid);
 	}
+	OutCryptoSettings.MasterEncryptionKey = OutCryptoSettings.EncryptionKeys.Find(EncryptionKeyOverrideGuid);
 }
 
-void ApplyKeyChain(const TKeyChain& KeyChain)
+void ApplyEncryptionKeys(const FKeyChain& KeyChain)
 {
-	if (KeyChain.Contains(FGuid()))
+	if (KeyChain.EncryptionKeys.Contains(FGuid()))
 	{
-		FAES::FAESKey DefaultKey = KeyChain[FGuid()].Key;
+		FAES::FAESKey DefaultKey = KeyChain.EncryptionKeys[FGuid()].Key;
 		FCoreDelegates::GetPakEncryptionKeyDelegate().BindLambda([DefaultKey](uint8 OutKey[32]) { FMemory::Memcpy(OutKey, DefaultKey.Key, sizeof(DefaultKey.Key)); });
 	}
 
-	for (const TKeyChain::ElementType& Key : KeyChain)
+	for (const TMap<FGuid, FNamedAESKey>::ElementType& Key : KeyChain.EncryptionKeys)
 	{
 		if (Key.Key.IsValid())
 		{
@@ -1437,30 +1438,32 @@ void ApplyKeyChain(const TKeyChain& KeyChain)
 /**
  * Creates a pak file writer. This can be a signed writer if the encryption keys are specified in the command line
  */
-FArchive* CreatePakWriter(const TCHAR* Filename, const FKeyPair& SigningKey)
+FArchive* CreatePakWriter(const TCHAR* Filename, const FKeyChain& InKeyChain, bool bSign)
 {
 	FArchive* Writer = IFileManager::Get().CreateFileWriter(Filename);
-	FString KeyFilename;
-	bool bSigningEnabled = false;
-	
+
 	if (Writer)
 	{
-		if (SigningKey.IsValid())
+		if (bSign)
 		{
 			UE_LOG(LogPakFile, Display, TEXT("Creating signed pak %s."), Filename);
-			Writer = new FSignedArchiveWriter(*Writer, Filename, SigningKey.PublicKey, SigningKey.PrivateKey);
+			Writer = new FSignedArchiveWriter(*Writer, Filename, InKeyChain.SigningKey.PublicKey, InKeyChain.SigningKey.PrivateKey);
+		}
+		else
+		{
+			UE_LOG(LogPakFile, Display, TEXT("Creating pak %s."), Filename);
 		}
 	}
 
 	return Writer;
 }
 
-bool CreatePakFile(const TCHAR* Filename, TArray<FPakInputPair>& FilesToAdd, const FPakCommandLineParameters& CmdLineParameters, const FKeyPair& SigningKey, const TKeyChain& KeyChain)
-{	
+bool CreatePakFile(const TCHAR* Filename, TArray<FPakInputPair>& FilesToAdd, const FPakCommandLineParameters& CmdLineParameters, const FKeyChain& InKeyChain)
+{
 	const double StartTime = FPlatformTime::Seconds();
 
 	// Create Pak
-	TUniquePtr<FArchive> PakFileHandle(CreatePakWriter(Filename, SigningKey));
+	TUniquePtr<FArchive> PakFileHandle(CreatePakWriter(Filename, InKeyChain, CmdLineParameters.bSign));
 	if (!PakFileHandle)
 	{
 		UE_LOG(LogPakFile, Error, TEXT("Unable to create pak file \"%s\"."), Filename);
@@ -1468,21 +1471,12 @@ bool CreatePakFile(const TCHAR* Filename, TArray<FPakInputPair>& FilesToAdd, con
 	}
 
 	FPakInfo Info;
-	FNamedAESKey MasterKey;
-	if (const FNamedAESKey* MasterKeyLookup = KeyChain.Find(CmdLineParameters.EncryptionKeyGuid))
-	{
-		MasterKey = *MasterKeyLookup;
-	}
-	Info.bEncryptedIndex = (MasterKey.IsValid() && CmdLineParameters.EncryptIndex);
-	Info.EncryptionKeyGuid = MasterKey.Guid;
+	Info.bEncryptedIndex = (InKeyChain.MasterEncryptionKey && CmdLineParameters.EncryptIndex);
+	Info.EncryptionKeyGuid = InKeyChain.MasterEncryptionKey ? InKeyChain.MasterEncryptionKey->Guid : FGuid();
 
-	if (CmdLineParameters.EncryptionKeyGuid.IsValid())
-	{
-		UE_LOG(LogPakFile, Display, TEXT("Encrypting using key '%s' [%s]"), *MasterKey.Name, *MasterKey.Guid.ToString());
-	}
-	else
+	if (InKeyChain.MasterEncryptionKey)
 	{
-		UE_LOG(LogPakFile, Display, TEXT("Encrypting using embedded key"));
+		UE_LOG(LogPakFile, Display, TEXT("Encrypting using key '%s' [%s]"), *InKeyChain.MasterEncryptionKey->Name, *InKeyChain.MasterEncryptionKey->Guid.ToString());
 	}
 
 	TArray<FPakEntryPair> Index;
@@ -1689,12 +1683,12 @@ bool CreatePakFile(const TCHAR* Filename, TArray<FPakInputPair>& FilesToAdd, con
 		}
 		else if (FilesToAdd[FileIndex].bNeedsCompression && CompressionMethod != NAME_None)
 		{
-			bCopiedToPak = PrepareCopyCompressedFileToPak(MountPoint, Info, FilesToAdd[FileIndex], CompressedFileBuffer, NewEntry, DataToWrite, SizeToWrite, MasterKey);
+			bCopiedToPak = PrepareCopyCompressedFileToPak(MountPoint, Info, FilesToAdd[FileIndex], CompressedFileBuffer, NewEntry, DataToWrite, SizeToWrite, InKeyChain);
 			DataToWrite = CompressedFileBuffer.CompressedBuffer.Get();
 		}
 		else
 		{
-			bCopiedToPak = PrepareCopyFileToPak(MountPoint, FilesToAdd[FileIndex], ReadBuffer, BufferSize, NewEntry, DataToWrite, SizeToWrite, MasterKey);
+			bCopiedToPak = PrepareCopyFileToPak(MountPoint, FilesToAdd[FileIndex], ReadBuffer, BufferSize, NewEntry, DataToWrite, SizeToWrite, InKeyChain);
 			DataToWrite = ReadBuffer;
 		}		
 
@@ -1781,7 +1775,7 @@ bool CreatePakFile(const TCHAR* Filename, TArray<FPakInputPair>& FilesToAdd, con
 			{
 				TotalRequestedEncryptedFiles++;
 
-				if (MasterKey.IsValid())
+				if (InKeyChain.MasterEncryptionKey)
 				{
 					TotalEncryptedFiles++;
 					TotalEncryptedDataSize += SizeToWrite;
@@ -1881,8 +1875,8 @@ bool CreatePakFile(const TCHAR* Filename, TArray<FPakInputPair>& FilesToAdd, con
 
 	if (Info.bEncryptedIndex)
 	{
-		FNamedAESKey MasterEncryptionKey = KeyChain[CmdLineParameters.EncryptionKeyGuid];
-		FAES::EncryptData(IndexData.GetData(), IndexData.Num(), MasterEncryptionKey.Key);
+		check(InKeyChain.MasterEncryptionKey);
+		FAES::EncryptData(IndexData.GetData(), IndexData.Num(), InKeyChain.MasterEncryptionKey->Key);
 		TotalEncryptedDataSize += IndexData.Num();
 	}
 
@@ -1933,9 +1927,9 @@ bool CreatePakFile(const TCHAR* Filename, TArray<FPakInputPair>& FilesToAdd, con
 	return true;
 }
 
-bool TestPakFile(const TCHAR* Filename, bool bSigned)
+bool TestPakFile(const TCHAR* Filename)
 {	
-	FPakFile PakFile(&FPlatformFileManager::Get().GetPlatformFile(), Filename, bSigned);
+	FPakFile PakFile(&FPlatformFileManager::Get().GetPlatformFile(), Filename, false);
 	if (PakFile.IsValid())
 	{
 		return PakFile.Check();
@@ -1947,9 +1941,9 @@ bool TestPakFile(const TCHAR* Filename, bool bSigned)
 	}
 }
 
-bool ListFilesInPak(const TCHAR * InPakFilename, int64 SizeFilter, bool bIncludeDeleted, const FString& CSVFilename, bool bExtractToMountPoint, bool bSigned)
+bool ListFilesInPak(const TCHAR * InPakFilename, int64 SizeFilter, bool bIncludeDeleted, const FString& CSVFilename, bool bExtractToMountPoint, const FKeyChain& InKeyChain)
 {
-	FPakFile PakFile(&FPlatformFileManager::Get().GetPlatformFile(), InPakFilename, bSigned);
+	FPakFile PakFile(&FPlatformFileManager::Get().GetPlatformFile(), InPakFilename, false);
 	int32 FileCount = 0;
 	int64 FileSize = 0;
 	int64 FilteredSize = 0;
@@ -2083,7 +2077,15 @@ bool ListFilesInPak(const TCHAR * InPakFilename, int64 SizeFilter, bool bInclude
 	}
 	else
 	{
-		UE_LOG(LogPakFile, Error, TEXT("Unable to open pak file \"%s\"."), InPakFilename);
+		if (PakFile.GetInfo().EncryptionKeyGuid.IsValid() && !InKeyChain.EncryptionKeys.Contains(PakFile.GetInfo().EncryptionKeyGuid))
+		{
+			UE_LOG(LogPakFile, Fatal, TEXT("Missing encryption key %s for pak file \"%s\"."), *PakFile.GetInfo().EncryptionKeyGuid.ToString(), InPakFilename);
+		}
+		else
+		{
+			UE_LOG(LogPakFile, Fatal, TEXT("Unable to open pak file \"%s\"."), InPakFilename);
+		}
+
 		return false;
 	}
 }
@@ -2133,7 +2135,7 @@ int32 GetPakChunkIndexFromFilename( const FString& PakFilePath )
 	return PakChunkIndex;
 }
 
-bool AuditPakFiles( const FString& InputPath, bool bOnlyDeleted, const FString& CSVFilename, bool bSigned, const TMap<FString, uint64>& OrderMap, bool bSortByOrdering )
+bool AuditPakFiles( const FString& InputPath, bool bOnlyDeleted, const FString& CSVFilename, const TMap<FString, uint64>& OrderMap, bool bSortByOrdering )
 {
 	//collect all pak files
 	FString PakFileDirectory;
@@ -2174,7 +2176,7 @@ bool AuditPakFiles( const FString& InputPath, bool bOnlyDeleted, const FString&
 		int32 PakPriority = GetPakPriorityFromFilename(PakFilename);
 		HighestPakPriority = FMath::Max( HighestPakPriority, PakPriority );
 
-		FPakFile PakFile(&FPlatformFileManager::Get().GetPlatformFile(), *PakFilename, bSigned);
+		FPakFile PakFile(&FPlatformFileManager::Get().GetPlatformFile(), *PakFilename, false);
 		if (PakFile.IsValid())
 		{
 			FString PakMountPoint = PakFile.GetMountPoint().Replace(TEXT("../../../"), TEXT(""));
@@ -2420,7 +2422,7 @@ bool AuditPakFiles( const FString& InputPath, bool bOnlyDeleted, const FString&
 	return true;
 }
 
-bool ListFilesAtOffset( const TCHAR* InPakFileName, const TArray<int64>& InOffsets, bool bSigned )
+bool ListFilesAtOffset( const TCHAR* InPakFileName, const TArray<int64>& InOffsets )
 {
 	if( InOffsets.Num() == 0 )
 	{
@@ -2428,7 +2430,7 @@ bool ListFilesAtOffset( const TCHAR* InPakFileName, const TArray<int64>& InOffse
 		return false;
 	}
 
-	FPakFile PakFile(&FPlatformFileManager::Get().GetPlatformFile(), InPakFileName, bSigned);
+	FPakFile PakFile(&FPlatformFileManager::Get().GetPlatformFile(), InPakFileName, false);
 	if (!PakFile.IsValid())
 	{
 		UE_LOG(LogPakFile, Error, TEXT("Failed to open %s"), *InPakFileName );
@@ -2485,8 +2487,6 @@ bool GeneratePIXMappingFile(const TArray<FString> InPakFileList, const FString&
 		FPlatformFileManager::Get().GetPlatformFile().CreateDirectoryTree(*OutputPath);
 	}
 
-	bool bPakFileSigned = FParse::Param(FCommandLine::Get(), TEXT("signed"));
-
 	for (const FString& PakFileName : InPakFileList)
 	{
 		// open CSV file, if requested
@@ -2502,7 +2502,7 @@ bool GeneratePIXMappingFile(const TArray<FString> InPakFileList, const FString&
 			}
 		}
 
-		FPakFile PakFile(&FPlatformFileManager::Get().GetPlatformFile(), *PakFileName, bPakFileSigned);
+		FPakFile PakFile(&FPlatformFileManager::Get().GetPlatformFile(), *PakFileName, false);
 		if (!PakFile.IsValid())
 		{
 			UE_LOG(LogPakFile, Error, TEXT("Failed to open %s"), *PakFileName);
@@ -2537,7 +2537,7 @@ struct FFileInfo
 	uint8 Hash[16];
 };
 
-bool ExtractFilesFromPak(const TCHAR* InPakFilename, TMap<FString, FFileInfo>& InFileHashes, const TCHAR* InDestPath, bool bUseMountPoint, const TKeyChain& InKeyChain, bool bSigned, const FString* InFilter, TArray<FPakInputPair>* OutEntries = nullptr, TArray<FPakInputPair>* OutDeletedEntries = nullptr, TMap<FString, uint64>* OutOrderMap = nullptr)
+bool ExtractFilesFromPak(const TCHAR* InPakFilename, TMap<FString, FFileInfo>& InFileHashes, const TCHAR* InDestPath, bool bUseMountPoint, const FKeyChain& InKeyChain, const FString* InFilter, TArray<FPakInputPair>* OutEntries = nullptr, TArray<FPakInputPair>* OutDeletedEntries = nullptr, TMap<FString, uint64>* OutOrderMap = nullptr, TArray<FGuid>* OutUsedEncryptionKeys = nullptr, bool* OutAnyPakSigned = nullptr)
 {
 	// Gather all patch versions of the requested pak file and run through each separately
 	TArray<FString> PakFileList;
@@ -2566,7 +2566,7 @@ bool ExtractFilesFromPak(const TCHAR* InPakFilename, TMap<FString, FFileInfo>& I
 		FString PakFilename = PakFileDirectory + "\\" + PakFileList[PakFileIndex];
 		int32 PakPriority = GetPakPriorityFromFilename(PakFilename);
 
-		FPakFile PakFile(&FPlatformFileManager::Get().GetPlatformFile(), *PakFilename, bSigned);
+		FPakFile PakFile(&FPlatformFileManager::Get().GetPlatformFile(), *PakFilename, false);
 		if (PakFile.IsValid())
 		{
 			FString DestPath(InDestPath);
@@ -2579,6 +2579,20 @@ bool ExtractFilesFromPak(const TCHAR* InPakFilename, TMap<FString, FFileInfo>& I
 			int32 FileCount = 0;
 			int32 ExtractedCount = 0;
 
+			if (OutUsedEncryptionKeys)
+			{
+				OutUsedEncryptionKeys->Add(PakFile.GetInfo().EncryptionKeyGuid);
+			}
+
+			if (OutAnyPakSigned && *OutAnyPakSigned == false)
+			{
+				FString SignatureFile(FPaths::ChangeExtension(PakFilename, TEXT(".sig")));
+				if (FPaths::FileExists(SignatureFile))
+				{
+					*OutAnyPakSigned = true;
+				}
+			}
+
 			FString PakMountPoint = bUseMountPoint ? PakFile.GetMountPoint().Replace(TEXT("../../../"), TEXT("")) : TEXT("");
 
 			for (FPakFile::FFileIterator It(PakFile,bIncludeDeleted); It; ++It, ++FileCount)
@@ -2687,7 +2701,7 @@ void CreateDiffRelativePathMap(TArray<FString>& FileNames, const FString& RootPa
 	}
 }
 
-bool DiffFilesInPaks(const FString& InPakFilename1, const FString& InPakFilename2, const bool bLogUniques1, const bool bLogUniques2, const TKeyChain& InKeyChain, const bool bSigned)
+bool DiffFilesInPaks(const FString& InPakFilename1, const FString& InPakFilename2, const bool bLogUniques1, const bool bLogUniques2, const FKeyChain& InKeyChain)
 {
 	int32 NumUniquePAK1 = 0;
 	int32 NumUniquePAK2 = 0;
@@ -2697,8 +2711,8 @@ bool DiffFilesInPaks(const FString& InPakFilename1, const FString& InPakFilename
 	TGuardValue<ELogTimes::Type> DisableLogTimes(GPrintLogTimes, ELogTimes::None);
 	UE_LOG(LogPakFile, Log, TEXT("FileEventType, FileName, Size1, Size2"));
 
-	FPakFile PakFile1(&FPlatformFileManager::Get().GetPlatformFile(), *InPakFilename1, bSigned);
-	FPakFile PakFile2(&FPlatformFileManager::Get().GetPlatformFile(), *InPakFilename2, bSigned);
+	FPakFile PakFile1(&FPlatformFileManager::Get().GetPlatformFile(), *InPakFilename1, false);
+	FPakFile PakFile2(&FPlatformFileManager::Get().GetPlatformFile(), *InPakFilename2, false);
 	if (PakFile1.IsValid() && PakFile2.IsValid())
 	{		
 		FArchive& PakReader1 = *PakFile1.GetSharedReader(NULL);
@@ -2859,7 +2873,7 @@ bool GenerateHashForFile( FString Filename, FFileInfo& FileHash)
 	return true;
 }
 
-bool GenerateHashesFromPak(const TCHAR* InPakFilename, const TCHAR* InDestPakFilename, TMap<FString, FFileInfo>& FileHashes, bool bUseMountPoint, const TKeyChain& InKeyChain, int32& OutLowestSourcePakVersion, const bool bSigned )
+bool GenerateHashesFromPak(const TCHAR* InPakFilename, const TCHAR* InDestPakFilename, TMap<FString, FFileInfo>& FileHashes, bool bUseMountPoint, const FKeyChain& InKeyChain, int32& OutLowestSourcePakVersion)
 {
 	OutLowestSourcePakVersion = FPakInfo::PakFile_Version_Invalid;
 
@@ -2885,7 +2899,7 @@ bool GenerateHashesFromPak(const TCHAR* InPakFilename, const TCHAR* InDestPakFil
 		int32 PakPriority = GetPakPriorityFromFilename(PakFilename);
 		int32 PakChunkIndex = GetPakChunkIndexFromFilename(PakFilename);
 
-		FPakFile PakFile(&FPlatformFileManager::Get().GetPlatformFile(), *PakFilename, bSigned);
+		FPakFile PakFile(&FPlatformFileManager::Get().GetPlatformFile(), *PakFilename, false);
 		if (PakFile.IsValid())
 		{
 			FArchive& PakReader = *PakFile.GetSharedReader(NULL);
@@ -3298,7 +3312,7 @@ void RemoveIdenticalFiles( TArray<FPakInputPair>& FilesToPak, const FString& Sou
 	FilesToPak.RemoveAt(WriteIndex, NumToRemove, true);
 }
 
-void ProcessLegacyFileMoves( TArray<FPakInputPair>& InDeleteRecords, TMap<FString, FFileInfo>& InExistingPackagedFileHashes, const FString& InInputPath, const TArray<FPakInputPair>& InFilesToPak, int32 CurrentPatchChunkIndex, bool bSigned )
+void ProcessLegacyFileMoves( TArray<FPakInputPair>& InDeleteRecords, TMap<FString, FFileInfo>& InExistingPackagedFileHashes, const FString& InInputPath, const TArray<FPakInputPair>& InFilesToPak, int32 CurrentPatchChunkIndex )
 {
 	double StartTime = FPlatformTime::Seconds();
 
@@ -3342,7 +3356,7 @@ void ProcessLegacyFileMoves( TArray<FPakInputPair>& InDeleteRecords, TMap<FStrin
 		UE_LOG(LogPakFile, Display, TEXT("Checking old pak file \"%s\" Pri:%d Chunk:%d."), *PakFilename, PakPriority, PakChunkIndex );
 
 
-		FPakFile PakFile(&FPlatformFileManager::Get().GetPlatformFile(), *PakFilename, bSigned);
+		FPakFile PakFile(&FPlatformFileManager::Get().GetPlatformFile(), *PakFilename, false);
 		if (PakFile.IsValid())
 		{
 			FString PakMountPoint = PakFile.GetMountPoint().Replace(TEXT("../../../"), TEXT(""));
@@ -3498,7 +3512,7 @@ FString GetPakPath(const TCHAR* SpecifiedPath, bool bIsForCreation)
 	return PakFilename;
 }
 
-bool Repack(const FString& InputPakFile, const FString& OutputPakFile, const FPakCommandLineParameters& CmdLineParameters, const FKeyPair& SigningKey, const TKeyChain& KeyChain, bool bIncludeDeleted, bool bSigned)
+bool Repack(const FString& InputPakFile, const FString& OutputPakFile, const FPakCommandLineParameters& CmdLineParameters, const FKeyChain& InKeyChain, bool bIncludeDeleted)
 {
 	bool bResult = false;
 
@@ -3507,8 +3521,10 @@ bool Repack(const FString& InputPakFile, const FString& OutputPakFile, const FPa
 	TArray<FPakInputPair> Entries;
 	TArray<FPakInputPair> DeletedEntries;
 	TMap<FString, uint64> OrderMap;
+	TArray<FGuid> EncryptionKeys;
+	bool bAnySigned = false;
 	FString TempDir = FPaths::EngineIntermediateDir() / TEXT("UnrealPak") / TEXT("Repack") / FPaths::GetBaseFilename(InputPakFile);
-	if (ExtractFilesFromPak(*InputPakFile, Hashes, *TempDir, false, KeyChain, bSigned, nullptr, &Entries, &DeletedEntries, &OrderMap))
+	if (ExtractFilesFromPak(*InputPakFile, Hashes, *TempDir, false, InKeyChain, nullptr, &Entries, &DeletedEntries, &OrderMap, &EncryptionKeys, &bAnySigned))
 	{
 		TArray<FPakInputPair> FilesToAdd;
 		CollectFilesToAdd(FilesToAdd, Entries, OrderMap);
@@ -3528,9 +3544,15 @@ bool Repack(const FString& InputPakFile, const FString& OutputPakFile, const FPa
 		// Get a temporary output filename. We'll only create/replace the final output file once successful.
 		FString TempOutputPakFile = FPaths::CreateTempFilename(*FPaths::GetPath(OutputPakFile), *FPaths::GetCleanFilename(OutputPakFile));
 
+		FPakCommandLineParameters ModifiedCmdLineParameters = CmdLineParameters;
+		ModifiedCmdLineParameters.bSign = bAnySigned && InKeyChain.SigningKey.IsValid();
+		
+		FKeyChain ModifiedKeyChain = InKeyChain;
+		ModifiedKeyChain.MasterEncryptionKey = InKeyChain.EncryptionKeys.Find(EncryptionKeys.Num() ? EncryptionKeys[0] : FGuid());
+
 		// Create the new pak file
 		UE_LOG(LogPakFile, Display, TEXT("Creating %s..."), *OutputPakFile);
-		if (CreatePakFile(*TempOutputPakFile, FilesToAdd, CmdLineParameters, SigningKey, KeyChain))
+		if (CreatePakFile(*TempOutputPakFile, FilesToAdd, ModifiedCmdLineParameters, ModifiedKeyChain))
 		{
 			IFileManager::Get().Move(*OutputPakFile, *TempOutputPakFile);
 
@@ -3560,10 +3582,6 @@ bool Repack(const FString& InputPakFile, const FString& OutputPakFile, const FPa
  *   -Test test if the pak file is healthy
  *   -Extract extracts pak file contents (followed by a path, i.e.: -extract D:\ExtractedPak)
  *   -Create=filename response file to create a pak file with
- *   -Sign=filename use the key pair in filename to sign a pak file, or: -sign=key_hex_values_separated_with_+, i.e: -sign=0x123456789abcdef+0x1234567+0x12345abc
- *    where the first number is the private key exponend, the second one is modulus and the third one is the public key exponent.
- *   -Signed use with -extract and -test to let the code know this is a signed pak
- *   -TableMax=number maximum prime number in the generated table (default is 10000)
  *
  * @param	ArgC	Command-line argument count
  * @param	ArgV	Argument strings
@@ -3581,11 +3599,6 @@ bool ExecuteUnrealPak(const TCHAR* CmdLine)
 		}
 	}
 
-	FKeyPair SigningKey;
-	TKeyChain KeyChain;
-	PrepareEncryptionAndSigningKeys(CmdLine, SigningKey, KeyChain);
-	ApplyKeyChain(KeyChain);
-
 	FString BatchFileName;
 	if (FParse::Value(CmdLine, TEXT("-Batch="), BatchFileName))
 	{
@@ -3607,6 +3620,11 @@ bool ExecuteUnrealPak(const TCHAR* CmdLine)
 		return Result;
 	}
 
+
+	FKeyChain KeyChain;
+	LoadKeyChain(CmdLine, KeyChain);
+	ApplyEncryptionKeys(KeyChain);
+
 	if (FParse::Param(CmdLine, TEXT("Test")))
 	{
 		if (NonOptionArguments.Num() != 1)
@@ -3616,8 +3634,7 @@ bool ExecuteUnrealPak(const TCHAR* CmdLine)
 		}
 
 		FString PakFilename = GetPakPath(*NonOptionArguments[0], false);
-		bool bSigned = FParse::Param(CmdLine, TEXT("signed"));
-		return TestPakFile(*PakFilename, bSigned);
+		return TestPakFile(*PakFilename);
 	}
 
 	if (FParse::Param(CmdLine, TEXT("List")))
@@ -3628,7 +3645,7 @@ bool ExecuteUnrealPak(const TCHAR* CmdLine)
 
 		if (NonOptionArguments.Num() != 1)
 		{
-			UE_LOG(LogPakFile, Error, TEXT("Incorrect arguments. Expected: -List <PakFile> [-SizeFilter=N] [-Signed]"));
+			UE_LOG(LogPakFile, Error, TEXT("Incorrect arguments. Expected: -List <PakFile> [-SizeFilter=N]"));
 			return false;
 		}
 
@@ -3638,14 +3655,13 @@ bool ExecuteUnrealPak(const TCHAR* CmdLine)
 		bool bExcludeDeleted = FParse::Param( CmdLine, TEXT("ExcludeDeleted") );
 
 		FString PakFilename = GetPakPath(*NonOptionArguments[0], false);
-		bool bSigned = FParse::Param(CmdLine, TEXT("signed"));
 
 		FString CSVFilename;
 		FParse::Value(CmdLine, TEXT("csv="), CSVFilename);
 
 		bool bExtractToMountPoint = FParse::Param(CmdLine, TEXT("ExtractToMountPoint"));
 
-		return ListFilesInPak(*PakFilename, SizeFilter, !bExcludeDeleted, *CSVFilename, bExtractToMountPoint, bSigned);
+		return ListFilesInPak(*PakFilename, SizeFilter, !bExcludeDeleted, *CSVFilename, bExtractToMountPoint, KeyChain);
 	}
 
 	if (FParse::Param(CmdLine, TEXT("Diff")))
@@ -3664,9 +3680,7 @@ bool ExecuteUnrealPak(const TCHAR* CmdLine)
 		const bool bLogUniques1 = bLogUniques && !FParse::Param(CmdLine, TEXT("nouniquesfile1"));
 		const bool bLogUniques2 = bLogUniques && !FParse::Param(CmdLine, TEXT("nouniquesfile2"));
 
-		const bool bSigned = FParse::Param(CmdLine, TEXT("signed"));
-
-		return DiffFilesInPaks(PakFilename1, PakFilename2, bLogUniques1, bLogUniques2, KeyChain, bSigned);
+		return DiffFilesInPaks(PakFilename1, PakFilename2, bLogUniques1, bLogUniques2, KeyChain);
 	}
 
 	if (FParse::Param(CmdLine, TEXT("Extract")))
@@ -3682,7 +3696,6 @@ bool ExecuteUnrealPak(const TCHAR* CmdLine)
 		}
 
 		FString PakFilename = GetPakPath(*NonOptionArguments[0], false);
-		bool bSigned = FParse::Param(CmdLine, TEXT("signed"));
 
 		bool bUseFilter = false;
 		FString Filter;
@@ -3691,7 +3704,7 @@ bool ExecuteUnrealPak(const TCHAR* CmdLine)
 		bUseFilter = FParse::Value(CmdLine, TEXT("Filter="), Filter);
 		bool bExtractToMountPoint = FParse::Param(CmdLine, TEXT("ExtractToMountPoint"));
 		TMap<FString, FFileInfo> EmptyMap;
-		return ExtractFilesFromPak(*PakFilename, EmptyMap, *DestPath, bExtractToMountPoint, KeyChain, bSigned, bUseFilter ? &Filter : nullptr);
+		return ExtractFilesFromPak(*PakFilename, EmptyMap, *DestPath, bExtractToMountPoint, KeyChain, bUseFilter ? &Filter : nullptr);
 	}
 
 	if (FParse::Param(CmdLine, TEXT("AuditFiles")))
@@ -3709,7 +3722,6 @@ bool ExecuteUnrealPak(const TCHAR* CmdLine)
 		FParse::Value( CmdLine, TEXT("CSV="), CSVFilename );
 
 		bool bOnlyDeleted = FParse::Param( CmdLine, TEXT("OnlyDeleted") );
-		bool bSigned = FParse::Param(CmdLine, TEXT("signed"));
 		bool bSortByOrdering = FParse::Param(CmdLine, TEXT("SortByOrdering"));
 
 		TMap<FString, uint64> OrderMap;
@@ -3724,7 +3736,7 @@ bool ExecuteUnrealPak(const TCHAR* CmdLine)
 			return false;
 		}
 
-		return AuditPakFiles(*PakFilenames, bOnlyDeleted, CSVFilename, bSigned, OrderMap, bSortByOrdering );
+		return AuditPakFiles(*PakFilenames, bOnlyDeleted, CSVFilename, OrderMap, bSortByOrdering );
 	}
 	
 	if (FParse::Param(CmdLine, TEXT("WhatsAtOffset")))
@@ -3746,9 +3758,7 @@ bool ExecuteUnrealPak(const TCHAR* CmdLine)
 			}
 		}
 
-		bool bSigned = FParse::Param(CmdLine, TEXT("signed"));
-
-		return ListFilesAtOffset( *PakFilename, Offsets, bSigned );
+		return ListFilesAtOffset( *PakFilename, Offsets );
 	}
 	
 	if (FParse::Param(CmdLine, TEXT("GeneratePIXMappingFile")))
@@ -3782,7 +3792,7 @@ bool ExecuteUnrealPak(const TCHAR* CmdLine)
 	{
 		if (NonOptionArguments.Num() != 1)
 		{
-			UE_LOG(LogPakFile, Error, TEXT("Incorrect arguments. Expected: -Repack <PakFile> [-Output=<PakFile>] [-Signed]"));
+			UE_LOG(LogPakFile, Error, TEXT("Incorrect arguments. Expected: -Repack <PakFile> [-Output=<PakFile>]"));
 			return false;
 		}
 
@@ -3794,11 +3804,18 @@ bool ExecuteUnrealPak(const TCHAR* CmdLine)
 		FString InputDir = FPaths::GetPath(*NonOptionArguments[0]);
 
 		TArray<FString> InputPakFiles;
-		IFileManager::Get().FindFiles(InputPakFiles, *InputDir, *FPaths::GetCleanFilename(*NonOptionArguments[0]));
-
-		for (int Idx = 0; Idx < InputPakFiles.Num(); Idx++)
+		if (FPaths::FileExists(*NonOptionArguments[0]))
 		{
-			InputPakFiles[Idx] = InputDir / InputPakFiles[Idx];
+			InputPakFiles.Add(*NonOptionArguments[0]);
+		}
+		else
+		{
+			IFileManager::Get().FindFiles(InputPakFiles, *InputDir, *FPaths::GetCleanFilename(*NonOptionArguments[0]));
+
+			for (int Idx = 0; Idx < InputPakFiles.Num(); Idx++)
+			{
+				InputPakFiles[Idx] = InputDir / InputPakFiles[Idx];
+			}
 		}
 
 		if (InputPakFiles.Num() == 0)
@@ -3836,11 +3853,10 @@ bool ExecuteUnrealPak(const TCHAR* CmdLine)
 		bool bExcludeDeleted = FParse::Param(CmdLine, TEXT("ExcludeDeleted"));
 
 		// Repack them all
-		bool bSigned = FParse::Param(CmdLine, TEXT("signed"));
 		for (int Idx = 0; Idx < InputPakFiles.Num(); Idx++)
 		{
 			UE_LOG(LogPakFile, Display, TEXT("Repacking %s into %s"), *InputPakFiles[Idx], *OutputPakFiles[Idx]);
-			if (!Repack(InputPakFiles[Idx], OutputPakFiles[Idx], CmdLineParameters, SigningKey, KeyChain, !bExcludeDeleted, bSigned))
+			if (!Repack(InputPakFiles[Idx], OutputPakFiles[Idx], CmdLineParameters, KeyChain, !bExcludeDeleted))
 			{
 				return false;
 			}
@@ -3853,7 +3869,6 @@ bool ExecuteUnrealPak(const TCHAR* CmdLine)
 	{
 		// since this is for creation, we pass true to make it not look in LaunchDir
 		FString PakFilename = GetPakPath(*NonOptionArguments[0], true);
-		bool bSigned = FParse::Param(CmdLine, TEXT("signed"));
 
 		// List of all items to add to pak file
 		TArray<FPakInputPair> Entries;
@@ -3895,20 +3910,19 @@ bool ExecuteUnrealPak(const TCHAR* CmdLine)
 			// Check command line for the "patchcryptokeys" param, which will tell us where to look for the encryption keys that
 			// we need to access the patch reference data
 			FString PatchReferenceCryptoKeysFilename;
-			TKeyChain PatchKeyChain = KeyChain;
+			FKeyChain PatchKeyChain;
 
 			if (FParse::Value(FCommandLine::Get(), TEXT("PatchCryptoKeys="), PatchReferenceCryptoKeysFilename))
 			{
-				FKeyPair UnusedSigningKey;
-				PrepareEncryptionAndSigningKeysFromCryptoKeyCache(PatchReferenceCryptoKeysFilename, UnusedSigningKey, PatchKeyChain);
-				ApplyKeyChain(PatchKeyChain);
+				LoadKeyChainFromFile(PatchReferenceCryptoKeysFilename, PatchKeyChain);
+				ApplyEncryptionKeys(PatchKeyChain);
 			}
 
 			UE_LOG(LogPakFile, Display, TEXT("Generating patch from %s."), *CmdLineParameters.SourcePatchPakFilename, true );
 
-			if (!GenerateHashesFromPak(*CmdLineParameters.SourcePatchPakFilename, *PakFilename, SourceFileHashes, true, PatchKeyChain, /*Out*/LowestSourcePakVersion, bSigned ))
+			if (!GenerateHashesFromPak(*CmdLineParameters.SourcePatchPakFilename, *PakFilename, SourceFileHashes, true, PatchKeyChain, /*Out*/LowestSourcePakVersion))
 			{
-				if (ExtractFilesFromPak(*CmdLineParameters.SourcePatchPakFilename, SourceFileHashes, *OutputPath, true, PatchKeyChain, bSigned, nullptr) == false)
+				if (ExtractFilesFromPak(*CmdLineParameters.SourcePatchPakFilename, SourceFileHashes, *OutputPath, true, PatchKeyChain, nullptr) == false)
 				{
 					UE_LOG(LogPakFile, Warning, TEXT("Unable to extract files from source pak file for patch"));
 				}
@@ -3918,7 +3932,7 @@ bool ExecuteUnrealPak(const TCHAR* CmdLine)
 				}
 			}
 
-			ApplyKeyChain(KeyChain);
+			ApplyEncryptionKeys(KeyChain);
 		}
 
 
@@ -3940,7 +3954,7 @@ bool ExecuteUnrealPak(const TCHAR* CmdLine)
 				FString SourcePakFolder = FPaths::GetPath(CmdLineParameters.SourcePatchPakFilename);
 
 				//remove invalid items from DeleteRecords and set 'bForceInclude' on some SourceFileHashes
-				ProcessLegacyFileMoves(DeleteRecords, SourceFileHashes, SourcePakFolder, FilesToAdd, CurrentPatchChunkIndex, bSigned);
+				ProcessLegacyFileMoves(DeleteRecords, SourceFileHashes, SourcePakFolder, FilesToAdd, CurrentPatchChunkIndex);
 			}
 			FilesToAdd.Append(DeleteRecords);
 
@@ -3949,7 +3963,7 @@ bool ExecuteUnrealPak(const TCHAR* CmdLine)
 		}
 
 
-		bool bResult = CreatePakFile(*PakFilename, FilesToAdd, CmdLineParameters, SigningKey, KeyChain);
+		bool bResult = CreatePakFile(*PakFilename, FilesToAdd, CmdLineParameters, KeyChain);
 
 		if (CmdLineParameters.GeneratePatch)
 		{
@@ -3973,7 +3987,6 @@ bool ExecuteUnrealPak(const TCHAR* CmdLine)
 	UE_LOG(LogPakFile, Error, TEXT("  UnrealPak <PakFolder> -AuditFiles [-OnlyDeleted] [-CSV=<filename>] [-order=<OrderingFile>] [-SortByOrdering]"));
 	UE_LOG(LogPakFile, Error, TEXT("  UnrealPak <PakFilename> -WhatsAtOffset [offset1] [offset2] [offset3] [...]"));
 	UE_LOG(LogPakFile, Error, TEXT("  UnrealPak <PakFolder> -GeneratePIXMappingFile -OutputPath=<Path>"));
-	UE_LOG(LogPakFile, Error, TEXT("  UnrealPak -TestEncryption"));
 	UE_LOG(LogPakFile, Error, TEXT("  Options:"));
 	UE_LOG(LogPakFile, Error, TEXT("    -blocksize=<BlockSize>"));
 	UE_LOG(LogPakFile, Error, TEXT("    -bitwindow=<BitWindow>"));
@@ -3987,5 +4000,7 @@ bool ExecuteUnrealPak(const TCHAR* CmdLine)
 	UE_LOG(LogPakFile, Error, TEXT("    -extracttomountpoint (Extract to mount point path of pak file)"));
 	UE_LOG(LogPakFile, Error, TEXT("    -encryptindex (encrypt the pak file index, making it unusable in unrealpak without supplying the key)"));
 	UE_LOG(LogPakFile, Error, TEXT("    -compressionformat[s]=<Format[,format2,...]> (set the format(s) to compress with, falling back on failures)"));
+	UE_LOG(LogPakFile, Error, TEXT("    -encryptionkeyoverrideguid (override the encryption key guid used for encrypting data in this pak file)"));
+	UE_LOG(LogPakFile, Error, TEXT("    -sign (generate a signature (.sig) file alongside the pak)"));
 	return false;
 }
diff --git a/Engine/Source/Programs/AutomationTool/Scripts/CopyBuildToStagingDirectory.Automation.cs b/Engine/Source/Programs/AutomationTool/Scripts/CopyBuildToStagingDirectory.Automation.cs
--- a/Engine/Source/Programs/AutomationTool/Scripts/CopyBuildToStagingDirectory.Automation.cs
+++ b/Engine/Source/Programs/AutomationTool/Scripts/CopyBuildToStagingDirectory.Automation.cs
@@ -137,14 +137,22 @@ public partial class Project : CommandUtils
 		{
 			CmdLine.AppendFormat(" -generatepatch={0} -tempfiles={1}", CommandUtils.MakePathSafeToUseWithCommandLine(PatchSourceContentPath), CommandUtils.MakePathSafeToUseWithCommandLine(CommandUtils.CombinePaths(CommandUtils.CmdEnv.LocalRoot, "TempFiles" + Path.GetFileNameWithoutExtension(OutputLocation.FullName))));
 		}
-		if (CryptoSettings != null && CryptoSettings.bEnablePakIndexEncryption)
+		if (CryptoSettings != null && CryptoSettings.bDataCryptoRequired)
 		{
-			CmdLine.AppendFormat(" -encryptindex");
-		}
-		if (!string.IsNullOrEmpty(EncryptionKeyGuid))
-		{
-			CmdLine.AppendFormat(" -EncryptionKeyOverrideGuid={0}", EncryptionKeyGuid);
+			if (CryptoSettings.bEnablePakIndexEncryption)
+			{
+				CmdLine.AppendFormat(" -encryptindex");
+			}
+			if (!string.IsNullOrEmpty(EncryptionKeyGuid))
+			{
+				CmdLine.AppendFormat(" -EncryptionKeyOverrideGuid={0}", EncryptionKeyGuid);
+			}
+			if (CryptoSettings.bDataCryptoRequired && CryptoSettings.bEnablePakSigning && CryptoSettings.SigningKey.IsValid())
+			{
+				CmdLine.AppendFormat(" -sign");
+			}
 		}
+		
 		CmdLine.Append(PlatformOptions);
 
 		return CmdLine.ToString();
diff --git a/Engine/Source/Programs/UnrealBuildTool/System/EncryptionAndSigning.cs b/Engine/Source/Programs/UnrealBuildTool/System/EncryptionAndSigning.cs
--- a/Engine/Source/Programs/UnrealBuildTool/System/EncryptionAndSigning.cs
+++ b/Engine/Source/Programs/UnrealBuildTool/System/EncryptionAndSigning.cs
@@ -31,11 +31,11 @@ namespace UnrealBuildTool
 			public byte[] Modulus;
 
 			/// <summary>
-			/// Returns TRUE if this is a valid Signing key
+			/// Determine if this key is valid
 			/// </summary>
 			public bool IsValid()
 			{
-				return Exponent != null && Modulus != null;
+				return Exponent != null && Modulus != null && Exponent.Length > 0 && Modulus.Length > 0;
 			}
 		}
 
@@ -55,7 +55,7 @@ namespace UnrealBuildTool
 			public SigningKey PrivateKey = new SigningKey();
 
 			/// <summary>
-			/// Returns TRUE if this is a valid signing key pair
+			/// Determine if this key is valid
 			/// </summary>
 			public bool IsValid()
 			{
@@ -77,16 +77,15 @@ namespace UnrealBuildTool
 			/// </summary>
 			public string Guid;
 			/// <summary>
-			/// 128 bit AES key
+			/// AES key
 			/// </summary>
 			public byte[] Key;
-
 			/// <summary>
-			/// Returns TRUE if this is a valid encryption key
+			/// Determine if this key is valid
 			/// </summary>
 			public bool IsValid()
 			{
-				return Key != null && Guid != null;
+				return Key != null && Key.Length > 0 && Guid != null;
 			}
 		}
 
@@ -396,6 +395,22 @@ namespace UnrealBuildTool
 				Settings = NewSettings;
 			}
 
+			// Validate the settings we have read
+			if (Settings.bDataCryptoRequired && Settings.bEnablePakSigning && (Settings.SigningKey == null || !Settings.SigningKey.IsValid()))
+			{
+				Log.TraceWarningOnce("Pak signing is enabled, but no valid signing keys were found. Please generate a key in the editor project crypto settings. Signing will be disabled");
+				Settings.bEnablePakSigning = false;
+			}
+
+			if (Settings.bDataCryptoRequired && Settings.IsAnyEncryptionEnabled() && (Settings.EncryptionKey == null || !Settings.EncryptionKey.IsValid()))
+			{
+				Log.TraceWarningOnce("Pak encryption is enabled, but no valid encryption key was found. Please generate a key in the editor project crypto settings. Encryption will be disabled");
+				Settings.bEnablePakUAssetEncryption = false;
+				Settings.bEnablePakFullAssetEncryption = false;
+				Settings.bEnablePakIndexEncryption = false;
+				Settings.bEnablePakIniEncryption = false;
+			}
+
 			return Settings;
 		}
 
