UnrealPak patch pak seek optimization v2
- Add support for incremental gap filling. Does gap filling in multiple passes, allows us to limit patch size growth to some fixed percentage
- Fixed cooker order UEXP files (previously these were unordered, which could actually result in patches increasing in size after gap filling)
- Added better method for measuring fragmentation - reports as a percentage, takes order index into account
- Refactored OrderMapping via new FPakOrderMap class which knows about primary/secondary mappings
- Added various schemes for defragmentation, including OnePass method which matches legacy behaviour, and various incremental schemes. Incremental_PrimaryFirst seems to give best results

diff --git a/Engine/Source/Developer/PakFileUtilities/Private/PakFileUtilities.cpp b/Engine/Source/Developer/PakFileUtilities/Private/PakFileUtilities.cpp
--- a/Engine/Source/Developer/PakFileUtilities/Private/PakFileUtilities.cpp
+++ b/Engine/Source/Developer/PakFileUtilities/Private/PakFileUtilities.cpp
@@ -23,6 +23,10 @@
 
 IMPLEMENT_MODULE(FDefaultModuleImpl, PakFileUtilities);
 
+#define GUARANTEE_UASSET_AND_UEXP_IN_SAME_PAK 1
+
+#define SEEK_OPT_VERBOSITY Display
+
 struct FNamedAESKey
 {
 	FString Name;
@@ -47,7 +51,7 @@ class FMemoryCompressor;
 /**
 * AsyncTask for FMemoryCompressor
 * Compress a memory block asynchronously
-*/
+ */
 class FBlockCompressTask : public FNonAbandonableTask
 {
 public:
@@ -79,7 +83,7 @@ public:
 	}
 	
 	FORCEINLINE TStatId GetStatId() const { RETURN_QUICK_DECLARE_CYCLE_STAT(ExampleAsyncTask, STATGROUP_ThreadPoolAsyncTasks); }
-
+		
 private:
 	// Source buffer
 	void* UncompressedBuffer;
@@ -105,7 +109,7 @@ public:
 	/** Divide into blocks and start compress asynchronously */
 	FMemoryCompressor(uint8* UncompressedBuffer, int32 UncompressedSize, FName Format, int32 CompressionBlockSize) :
 		Index(0)
-	{
+			{
 		// Divide into blocks and start compression async tasks.
 		// These blocks must be as same as followed CompressMemory callings.
 		int64 UncompressedBytes = 0;
@@ -123,19 +127,19 @@ public:
 	~FMemoryCompressor()
 	{
 		for (auto* AsyncTask : BlockCompressAsyncTasks)
-		{
-			if (!AsyncTask->Cancel())
 			{
+			if (!AsyncTask->Cancel())
+				{
 				AsyncTask->EnsureCompletion();
-			}
+				}
 			delete AsyncTask;
+			}
 		}
-	}
 
 
-	/** Fetch compressed result. Returns true and store CompressedSize if succeeded */
-	bool CompressMemory(FName Format, void* CompressedBuffer, int32& CompressedSize, const void* UncompressedBuffer, int32 UncompressedSize)
-	{
+/** Fetch compressed result. Returns true and store CompressedSize if succeeded */
+bool CompressMemory(FName Format, void* CompressedBuffer, int32& CompressedSize, const void* UncompressedBuffer, int32 UncompressedSize)
+{
 		// Fetch compressed result from task.
 		// We assume this is called only once, same order, same parameters for
 		// each task.
@@ -148,7 +152,7 @@ public:
 		check(Task.UncompressedSize == UncompressedSize);
 		check(CompressedSize >= Task.CompressedSize);
 		if (!Task.Result)
-		{
+	{
 			return false;
 		}
 		FMemory::Memcpy(CompressedBuffer, Task.CompressedBuffer, Task.CompressedSize);
@@ -164,8 +168,188 @@ private:
 	int32 Index;
 };
 
+/**
+* Defines the order mapping for files within a pak
+ */
+class FPakOrderMap
+{
+public:
+	FPakOrderMap()
+		: MaxPrimaryOrderIndex(MAX_uint64)
+	{}
+
+	int32 Num() const
+	{
+		return OrderMap.Num();
+	}
 
+	void Add(const FString& Filename, uint64 Index)
+	{
+		OrderMap.Add(Filename, Index);
+	}
+
+	bool ProcessOrderFile(const TCHAR* ResponseFile, bool bSecondaryOrderFile = false)
+	{
+		int32 OrderOffset = 0;
+		if (bSecondaryOrderFile)
+		{
+			OrderOffset = Num();
+			MaxPrimaryOrderIndex = OrderOffset;
+		}
+		// List of all items to add to pak file
+		FString Text;
+		UE_LOG(LogPakFile, Display, TEXT("Loading pak order file %s..."), ResponseFile);
+		if (FFileHelper::LoadFileToString(Text, ResponseFile))
+		{
+			// Read all lines
+			TArray<FString> Lines;
+			Text.ParseIntoArray(Lines, TEXT("\n"), true);
+			for (int32 EntryIndex = 0; EntryIndex < Lines.Num(); EntryIndex++)
+			{
+				Lines[EntryIndex].ReplaceInline(TEXT("\r"), TEXT(""));
+				Lines[EntryIndex].ReplaceInline(TEXT("\n"), TEXT(""));
+				int32 OpenOrderNumber = EntryIndex;
+				if (Lines[EntryIndex].FindLastChar('"', OpenOrderNumber))
+				{
+					FString ReadNum = Lines[EntryIndex].RightChop(OpenOrderNumber + 1);
+					Lines[EntryIndex] = Lines[EntryIndex].Left(OpenOrderNumber + 1);
+					ReadNum.TrimStartInline();
+					if (ReadNum.IsNumeric())
+					{
+						OpenOrderNumber = FCString::Atoi(*ReadNum);
+					}
+				}
+				Lines[EntryIndex] = Lines[EntryIndex].TrimQuotes();
+				FString Path = FString::Printf(TEXT("%s"), *Lines[EntryIndex]);
+				FPaths::NormalizeFilename(Path);
+				Path = Path.ToLower();
+				if (bSecondaryOrderFile && OrderMap.Contains(Path))
+				{
+					continue;
+				}
+				OrderMap.Add(Path, OpenOrderNumber + OrderOffset);
+			}
+			UE_LOG(LogPakFile, Display, TEXT("Finished loading pak order file %s."), ResponseFile);
+			return true;
+		}
+		else
+		{
+			UE_LOG(LogPakFile, Error, TEXT("Unable to load pak order file %s."), ResponseFile);
+			return false;
+		}
+	}
 
+	uint64 GetFileOrder(const FString& Path, bool bAllowUexpUBulkFallback, bool* OutIsPrimary=nullptr) const
+	{
+		FString RegionStr;
+		FString NewPath = RemapLocalizationPathIfNeeded(Path.ToLower(), RegionStr);
+		const uint64* FoundOrder = OrderMap.Find(NewPath);
+		uint64 ReturnOrder = MAX_uint64;
+		if (FoundOrder != nullptr)
+		{
+			ReturnOrder = *FoundOrder;
+			if (OutIsPrimary)
+			{
+				*OutIsPrimary = ( ReturnOrder < MaxPrimaryOrderIndex );
+			}
+		}
+		else if (bAllowUexpUBulkFallback)
+		{
+			// if this is a cook order or an old order it will not have uexp files in it, so we put those in the same relative order after all of the normal files, but before any ubulk files
+			if (Path.EndsWith(TEXT("uexp")) || Path.EndsWith(TEXT("ubulk")))
+			{
+				uint64 CounterpartOrder = GetFileOrder(FPaths::GetBaseFilename(Path, false) + TEXT(".uasset"), false);
+				if (CounterpartOrder == MAX_uint64)
+				{
+					CounterpartOrder = GetFileOrder(FPaths::GetBaseFilename(Path, false) + TEXT(".umap"), false);
+				}
+				if (CounterpartOrder != MAX_uint64)
+				{
+					if (Path.EndsWith(TEXT("uexp")))
+					{
+						ReturnOrder = CounterpartOrder | (1 << 29);
+					}
+					else
+					{
+						ReturnOrder = CounterpartOrder | (1 << 30);
+					}
+				}
+			}
+		}
+
+		// Optionally offset based on region, so multiple files in different regions don't get the same order.
+		// I/O profiling suggests this is slightly worse, so leaving this disabled for now
+#if 0
+		if (ReturnOrder != MAX_uint64)
+		{
+			if (RegionStr.Len() > 0)
+			{
+				uint64 RegionOffset = 0;
+				for (int i = 0; i < RegionStr.Len(); i++)
+				{
+					int8 Letter = (int8)(RegionStr[i] - TEXT('a'));
+					RegionOffset |= (uint64(Letter) << (i * 5));
+				}
+				return ReturnOrder + (RegionOffset << 16);
+			}
+		}
+#endif
+		return ReturnOrder;
+	}
+private:
+	FString RemapLocalizationPathIfNeeded(const FString& PathLower, FString& OutRegion) const
+	{
+		static const TCHAR* L10NPrefix = (const TCHAR*)TEXT("/content/l10n/");
+		static const int32 L10NPrefixLength = FCString::Strlen(L10NPrefix);
+		int32 FoundIndex = PathLower.Find(L10NPrefix, ESearchCase::CaseSensitive);
+		if (FoundIndex > 0)
+		{
+			// Validate the content index is the first one
+			int32 ContentIndex = PathLower.Find(TEXT("/content/"), ESearchCase::CaseSensitive);
+			if (ContentIndex == FoundIndex)
+			{
+				int32 EndL10NOffset = ContentIndex + L10NPrefixLength;
+				int32 NextSlashIndex = PathLower.Find(TEXT("/"), ESearchCase::CaseSensitive, ESearchDir::FromStart, EndL10NOffset);
+				int32 RegionLength = NextSlashIndex - EndL10NOffset;
+				if (RegionLength >= 2)
+				{
+					FString NonLocalizedPath = PathLower.Mid(0, ContentIndex) + TEXT("/content") + PathLower.Mid(NextSlashIndex);
+					OutRegion = PathLower.Mid(EndL10NOffset, RegionLength);
+					return NonLocalizedPath;
+				}
+			}
+		}
+		return PathLower;
+	}
+
+	TMap<FString, uint64> OrderMap;
+	uint64 MaxPrimaryOrderIndex;
+};
+
+
+enum class ESeekOptMode : uint8
+{
+	None = 0,
+	OnePass = 1,
+	Incremental = 2,
+	Incremental_OnlyPrimaryOrder = 3,
+	Incremental_PrimaryThenSecondary = 4,
+	COUNT
+};
+
+struct FPatchSeekOptParams
+{
+	FPatchSeekOptParams()
+		: MaxGapSize(0)
+		, MaxInflationPercent(0.0f)
+		, Mode(ESeekOptMode::None)
+		, MaxAdjacentOrderDiff(128)
+	{}
+	int64 MaxGapSize;
+	float MaxInflationPercent; // For Incremental_ modes only
+	ESeekOptMode Mode;
+	int32 MaxAdjacentOrderDiff;
+};
 
 struct FPakCommandLineParameters
 {
@@ -175,8 +359,6 @@ struct FPakCommandLineParameters
 		, PatchFilePadAlign(0)
 		, AlignForMemoryMapping(0)
 		, GeneratePatch(false)
-		, PatchSeekOptMaxGapSize(0)
-		, PatchSeekOptUseOrder(false)
 		, EncryptIndex(false)
 		, UseCustomCompressor(false)
 		, bSign(false)
@@ -185,6 +367,7 @@ struct FPakCommandLineParameters
 	}
 
 	TArray<FName> CompressionFormats;
+	FPatchSeekOptParams SeekOptParams;
 	int32  CompressionBlockSize;
 	int64  FileSystemBlockSize;
 	int64  PatchFilePadAlign;
@@ -192,8 +375,6 @@ struct FPakCommandLineParameters
 	bool   GeneratePatch;
 	FString SourcePatchPakFilename;
 	FString SourcePatchDiffDirectory;
-	int64 PatchSeekOptMaxGapSize;
-	bool PatchSeekOptUseOrder;
 	bool EncryptIndex;
 	bool UseCustomCompressor;
 	FGuid EncryptionKeyGuid;
@@ -215,12 +396,14 @@ struct FPakInputPair
 	bool bNeedsCompression;
 	bool bNeedEncryption;
 	bool bIsDeleteRecord;	// This is used for patch PAKs when a file is deleted from one patch to the next
+	bool bIsInPrimaryOrder;
 
 	FPakInputPair()
 		: SuggestedOrder(MAX_uint64)
 		, bNeedsCompression(false)
 		, bNeedEncryption(false)
 		, bIsDeleteRecord(false)
+		, bIsInPrimaryOrder(false)
 	{}
 
 	FPakInputPair(const FString& InSource, const FString& InDest)
@@ -553,60 +736,6 @@ void PrepareDeleteRecordForPak(const FString& InMountPoint, const FPakInputPair
 	OutNewEntry.Info.SetDeleteRecord(true);
 }
 
-bool ProcessOrderFile(const TCHAR* ResponseFile, TMap<FString, uint64>& OrderMap, bool bSecondaryOrderFile = false, int32 OrderOffset = 0)
-{
-	// List of all items to add to pak file
-	FString Text;
-	UE_LOG(LogPakFile, Display, TEXT("Loading pak order file %s..."), ResponseFile);
-	if (FFileHelper::LoadFileToString(Text, ResponseFile))
-	{
-		// Read all lines
-		TArray<FString> Lines;
-		Text.ParseIntoArray(Lines, TEXT("\n"), true);
-		for (int32 EntryIndex = 0; EntryIndex < Lines.Num(); EntryIndex++)
-		{
-			Lines[EntryIndex].ReplaceInline(TEXT("\r"), TEXT(""));
-			Lines[EntryIndex].ReplaceInline(TEXT("\n"), TEXT(""));
-			int32 OpenOrderNumber = EntryIndex;
-			if (Lines[EntryIndex].FindLastChar('"', OpenOrderNumber))
-			{
-				FString ReadNum = Lines[EntryIndex].RightChop(OpenOrderNumber+1);
-				Lines[EntryIndex] = Lines[EntryIndex].Left(OpenOrderNumber+1);
-				ReadNum.TrimStartInline();
-				if (ReadNum.IsNumeric())
-				{
-					OpenOrderNumber = FCString::Atoi(*ReadNum);
-				}
-			}
-			Lines[EntryIndex] = Lines[EntryIndex].TrimQuotes();
-			FString Path=FString::Printf(TEXT("%s"), *Lines[EntryIndex]);
-			FPaths::NormalizeFilename(Path);
-			Path = Path.ToLower();
-#if 0
-			if (Path.EndsWith("uexp"))
-			{
-				OpenOrderNumber += (1 << 29);
-			}
-			if (Path.EndsWith("ubulk"))
-			{
-				OpenOrderNumber += (1 << 30);
-			}
-#endif
-			if (bSecondaryOrderFile && OrderMap.Contains(Path) )
-			{
-				continue;
-			}
-			OrderMap.Add(Path, OpenOrderNumber + OrderOffset);
-		}
-		UE_LOG(LogPakFile, Display, TEXT("Finished loading pak order file %s."), ResponseFile);
-		return true;
-	}
-	else 
-	{
-		UE_LOG(LogPakFile, Error, TEXT("Unable to load pak order file %s."), ResponseFile);
-		return false;
-	}
-}
 
 static void CommandLineParseHelper(const TCHAR* InCmdLine, TArray<FString>& Tokens, TArray<FString>& Switches)
 {
@@ -633,11 +762,11 @@ void ProcessCommandLine(const TCHAR* CmdLine, const TArray<FString>& NonOptionAr
 	if (FParse::Param(CmdLine, TEXT("patchcompatibilitymode421")))
 	{
 		CmdLineParameters.bPatchCompatibilityMode421 = true;
-	}
+		}
 
 	if (FParse::Value(CmdLine, TEXT("-blocksize="), ClusterSizeString) && 
 		FParse::Value(CmdLine, TEXT("-blocksize="), CmdLineParameters.FileSystemBlockSize))
-	{
+		{
 		if (ClusterSizeString.EndsWith(TEXT("MB")))
 		{
 			CmdLineParameters.FileSystemBlockSize *= 1024*1024;
@@ -650,18 +779,18 @@ void ProcessCommandLine(const TCHAR* CmdLine, const TArray<FString>& NonOptionAr
 	else
 	{
 		CmdLineParameters.FileSystemBlockSize = 0;
-	}
+}
 
 	FString CompBlockSizeString;
 	if (FParse::Value(CmdLine, TEXT("-compressionblocksize="), CompBlockSizeString) &&
 		FParse::Value(CmdLine, TEXT("-compressionblocksize="), CmdLineParameters.CompressionBlockSize))
-	{
+{
 		if (CompBlockSizeString.EndsWith(TEXT("MB")))
-		{
+	{
 			CmdLineParameters.CompressionBlockSize *= 1024 * 1024;
-		}
+	}
 		else if (CompBlockSizeString.EndsWith(TEXT("KB")))
-		{
+	{
 			CmdLineParameters.CompressionBlockSize *= 1024;
 		}
 	}
@@ -698,7 +827,7 @@ void ProcessCommandLine(const TCHAR* CmdLine, const TArray<FString>& NonOptionAr
 			FName FormatName = *Format;
 
 			if (FCompression::IsFormatValid(FormatName))
-			{
+	{
 				CmdLineParameters.CompressionFormats.Add(FormatName);
 				break;
 			}
@@ -716,8 +845,16 @@ void ProcessCommandLine(const TCHAR* CmdLine, const TArray<FString>& NonOptionAr
 
 		if (CmdLineParameters.GeneratePatch)
 		{
-			ReadSizeParam(CmdLine, TEXT("-patchSeekOptMaxGapSize="), CmdLineParameters.PatchSeekOptMaxGapSize);
-			CmdLineParameters.PatchSeekOptUseOrder = FParse::Param(CmdLine, TEXT("patchSeekOptUseOrder"));
+			FParse::Value(CmdLine, TEXT("-patchSeekOptMaxInflationPercent="), CmdLineParameters.SeekOptParams.MaxInflationPercent);
+			ReadSizeParam(CmdLine, TEXT("-patchSeekOptMaxGapSize="), CmdLineParameters.SeekOptParams.MaxGapSize);
+			FParse::Value(CmdLine, TEXT("-patchSeekOptMaxAdjacentOrderDiff="), CmdLineParameters.SeekOptParams.MaxAdjacentOrderDiff);
+
+			// For legacy reasons, if we specify a max gap size without a mode, we default to OnePass
+			if (CmdLineParameters.SeekOptParams.MaxGapSize > 0)
+			{
+				CmdLineParameters.SeekOptParams.Mode = ESeekOptMode::OnePass;
+			}
+			FParse::Value(CmdLine, TEXT("-patchSeekOptMode="), (int32&)CmdLineParameters.SeekOptParams.Mode);
 		}
 
 		bool bCompress = FParse::Param(CmdLine, TEXT("compress"));
@@ -839,60 +976,7 @@ void ProcessCommandLine(const TCHAR* CmdLine, const TArray<FString>& NonOptionAr
 	UE_LOG(LogPakFile, Display, TEXT("Added %d entries to add to pak file."), Entries.Num());
 }
 
-FString RemapLocalizationPathIfNeeded(const FString& PathLower, FString& OutRegion)
-{
-	static const TCHAR* L10NPrefix = (const TCHAR*)TEXT("/content/l10n/");
-	static const int32 L10NPrefixLength = FCString::Strlen(L10NPrefix);
-	int32 FoundIndex = PathLower.Find(L10NPrefix, ESearchCase::CaseSensitive);
-	if (FoundIndex > 0)
-	{
-		// Validate the content index is the first one
-		int32 ContentIndex = PathLower.Find(TEXT("/content/"), ESearchCase::CaseSensitive);
-		if (ContentIndex == FoundIndex)
-		{
-			int32 EndL10NOffset = ContentIndex + L10NPrefixLength;
-			int32 NextSlashIndex = PathLower.Find(TEXT("/"), ESearchCase::CaseSensitive, ESearchDir::FromStart, EndL10NOffset);
-			int32 RegionLength = NextSlashIndex - EndL10NOffset;
-			if (RegionLength >= 2)
-			{
-				FString NonLocalizedPath = PathLower.Mid(0, ContentIndex) + TEXT("/content") + PathLower.Mid(NextSlashIndex);
-				OutRegion = PathLower.Mid(EndL10NOffset, RegionLength);
-				return NonLocalizedPath;
-			}
-		}
-	}
-	return PathLower;
-}
-
-uint64 GetFileOrder(const FString Path, const TMap<FString, uint64>& OrderMap)
-{
-	FString RegionStr;
-	FString NewPath = RemapLocalizationPathIfNeeded(Path.ToLower(), RegionStr);
-	const uint64* FoundOrder = OrderMap.Find(NewPath);
-	if (FoundOrder == nullptr)
-	{
-		return MAX_uint64;
-	}
-
-	// Optionally offset based on region, so multiple files in different regions don't get the same order.
-	// I/O profiling suggests this is slightly worse, so leaving this disabled for now
-#if 0
-	if (RegionStr.Len() > 0)
-	{
-		uint64 RegionOffset = 0;
-		for (int i = 0; i < RegionStr.Len(); i++)
-		{
-			int8 Letter = (int8)(RegionStr[i] - TEXT('a'));
-			RegionOffset |= (uint64(Letter) << (i*5));
-		}
-		return *FoundOrder + (RegionOffset << 16);
-	}
-#endif
-	return *FoundOrder;
-}
-
-
-void CollectFilesToAdd(TArray<FPakInputPair>& OutFilesToAdd, const TArray<FPakInputPair>& InEntries, const TMap<FString, uint64>& OrderMap)
+void CollectFilesToAdd(TArray<FPakInputPair>& OutFilesToAdd, const TArray<FPakInputPair>& InEntries, const FPakOrderMap& OrderMap)
 {
 	UE_LOG(LogPakFile, Display, TEXT("Collecting files to add to pak file..."));
 	const double StartTime = FPlatformTime::Seconds();
@@ -929,7 +1013,7 @@ void CollectFilesToAdd(TArray<FPakInputPair>& OutFilesToAdd, const TArray<FPakIn
 				FPaths::MakeStandardFilename(FileInput.Source);
 				FileInput.Dest = FileInput.Source.Replace(*Directory, *Input.Dest, ESearchCase::IgnoreCase);
 				
-				uint64 FileOrder = GetFileOrder(FileInput.Dest, OrderMap);
+				uint64 FileOrder = OrderMap.GetFileOrder(FileInput.Dest, false, &FileInput.bIsInPrimaryOrder);
 				if(FileOrder != MAX_uint64)
 				{
 					FileInput.SuggestedOrder = FileOrder;
@@ -941,10 +1025,10 @@ void CollectFilesToAdd(TArray<FPakInputPair>& OutFilesToAdd, const TArray<FPakIn
 					// if this is a cook order or an old order it will not have uexp files in it, so we put those in the same relative order after all of the normal files, but before any ubulk files
 					if (FileInput.Dest.EndsWith(TEXT("uexp")) || FileInput.Dest.EndsWith(TEXT("ubulk")))
 					{
-						FileOrder = GetFileOrder(FPaths::GetBaseFilename(FileInput.Dest, false) + TEXT(".uasset"), OrderMap);
+						FileOrder = OrderMap.GetFileOrder(FPaths::GetBaseFilename(FileInput.Dest, false) + TEXT(".uasset"), false, &FileInput.bIsInPrimaryOrder);
 						if (FileOrder == MAX_uint64)
 						{
-							FileOrder = GetFileOrder(FPaths::GetBaseFilename(FileInput.Dest, false) + TEXT(".umap"), OrderMap);
+							FileOrder = OrderMap.GetFileOrder(FPaths::GetBaseFilename(FileInput.Dest, false) + TEXT(".umap"),  false, &FileInput.bIsInPrimaryOrder);
 						}
 						if (FileInput.Dest.EndsWith(TEXT("uexp")))
 						{
@@ -980,7 +1064,7 @@ void CollectFilesToAdd(TArray<FPakInputPair>& OutFilesToAdd, const TArray<FPakIn
 			FileInput.Source = Input.Source;
 			FPaths::MakeStandardFilename(FileInput.Source);
 			FileInput.Dest = FileInput.Source.Replace(*Directory, *Input.Dest, ESearchCase::IgnoreCase);
-			uint64 FileOrder = GetFileOrder(FileInput.Dest, OrderMap);
+			uint64 FileOrder = OrderMap.GetFileOrder(FileInput.Dest, true, &FileInput.bIsInPrimaryOrder);
 			if (FileOrder != MAX_uint64)
 			{
 				FileInput.SuggestedOrder = FileOrder;
@@ -1442,7 +1526,7 @@ bool CreatePakFile(const TCHAR* Filename, TArray<FPakInputPair>& FilesToAdd, con
 		Info.GetCompressionMethodIndex(NAME_Gzip);
 		Info.GetCompressionMethodIndex(TEXT("Bogus"));
 		Info.GetCompressionMethodIndex(TEXT("Oodle"));
-		
+
 	}
 
 
@@ -1501,7 +1585,7 @@ bool CreatePakFile(const TCHAR* Filename, TArray<FPakInputPair>& FilesToAdd, con
 		if (FileIndex)
 		{
 			if (FPaths::GetBaseFilename(FilesToAdd[FileIndex - 1].Dest, false) == FPaths::GetBaseFilename(FilesToAdd[FileIndex].Dest, false) &&
-				FPaths::GetExtension(FilesToAdd[FileIndex - 1].Dest, true) == TEXT(".uasset") && 
+				FPaths::GetExtension(FilesToAdd[FileIndex - 1].Dest, true) == TEXT(".uasset") &&
 				FPaths::GetExtension(FilesToAdd[FileIndex].Dest, true) == TEXT(".uexp")
 				)
 			{
@@ -1551,7 +1635,7 @@ bool CreatePakFile(const TCHAR* Filename, TArray<FPakInputPair>& FilesToAdd, con
 					{
 						// Check the compression ratio, if it's too low just store uncompressed. Also take into account read size
 						// if we still save 64KB it's probably worthwhile compressing, as that saves a file read operation in the runtime.
-								// TODO: drive this threshold from the command line
+									// TODO: drive this threshold from the command line
 						float PercentLess = ((float)CompressedFileBuffer.TotalCompressedSize / (OriginalFileSize / 100.f));
 						if (PercentLess > 90.f && (OriginalFileSize - CompressedFileBuffer.TotalCompressedSize) < 65536)
 						{
@@ -1662,7 +1746,7 @@ bool CreatePakFile(const TCHAR* Filename, TArray<FPakInputPair>& FilesToAdd, con
 		{
 			bCopiedToPak = PrepareCopyFileToPak(MountPoint, FilesToAdd[FileIndex], ReadBuffer, BufferSize, NewEntry, DataToWrite, SizeToWrite, InKeyChain);
 			DataToWrite = ReadBuffer;
-		}		
+		}
 
 		int64 TotalSizeToWrite = SizeToWrite + NewEntry.Info.GetSerializedSize(FPakInfo::PakFile_Version_Latest);
 		if (bCopiedToPak)
@@ -1721,7 +1805,7 @@ bool CreatePakFile(const TCHAR* Filename, TArray<FPakInputPair>& FilesToAdd, con
 				else
 				{
 					ContiguousTotalSizeSmallerThanBlockSize += TotalSizeToWrite;
-					ContiguousFilesSmallerThanBlockSize++;				
+					ContiguousFilesSmallerThanBlockSize++;
 				}
 			}
 			if (FilesToAdd[FileIndex].bNeedsCompression && CompressionMethod != NAME_None)
@@ -1821,7 +1905,7 @@ bool CreatePakFile(const TCHAR* Filename, TArray<FPakInputPair>& FilesToAdd, con
 			if (TotalSizeToWrite >= RequiredPatchPadding)
 			{
 				int64 RealStart = Entry.Info.Offset;
-				if ((RealStart % RequiredPatchPadding) != 0 && 
+				if ((RealStart % RequiredPatchPadding) != 0 &&
 					!Entry.Filename.EndsWith(TEXT("uexp")) && // these are export sections of larger files and may be packed with uasset/umap and so we don't need a warning here
 					!(Entry.Filename.EndsWith(TEXT(".m.ubulk")) && CmdLineParameters.AlignForMemoryMapping > 0)) // Bulk padding unaligns patch padding and so we don't need a warning here
 				{
@@ -1879,7 +1963,7 @@ bool CreatePakFile(const TCHAR* Filename, TArray<FPakInputPair>& FilesToAdd, con
 		{
 			UE_LOG(LogPakFile, Display, TEXT("  Index: Unencrypted"));
 		}
-		
+
 
 		UE_LOG(LogPakFile, Display, TEXT("  Total: %d bytes (%.2fMB)"), TotalEncryptedDataSize, (float)TotalEncryptedDataSize / 1024.0f / 1024.0f);
 	}
@@ -2107,7 +2191,7 @@ int32 GetPakChunkIndexFromFilename( const FString& PakFilePath )
 	return PakChunkIndex;
 }
 
-bool AuditPakFiles( const FString& InputPath, bool bOnlyDeleted, const FString& CSVFilename, const TMap<FString, uint64>& OrderMap, bool bSortByOrdering )
+bool AuditPakFiles( const FString& InputPath, bool bOnlyDeleted, const FString& CSVFilename, const FPakOrderMap& OrderMap, bool bSortByOrdering )
 {
 	//collect all pak files
 	FString PakFileDirectory;
@@ -2237,9 +2321,9 @@ bool AuditPakFiles( const FString& InputPath, bool bOnlyDeleted, const FString&
 			FPaths::NormalizeFilename(OpenOrderAssetName);
 			OpenOrderAssetName.ToLowerInline();
 
-			if( const uint64* OrderIndexPtr = OrderMap.Find( OpenOrderAssetName ) )
+			if( uint64 OrderIndex = OrderMap.GetFileOrder( OpenOrderAssetName, false ) != MAX_uint64 )
 			{
-				CachedOpenOrder.Add( AssetPath, *OrderIndexPtr );
+				CachedOpenOrder.Add( AssetPath, OrderIndex );
 			}
 		}
 	}
@@ -2509,7 +2593,7 @@ struct FFileInfo
 	uint8 Hash[16];
 };
 
-bool ExtractFilesFromPak(const TCHAR* InPakFilename, TMap<FString, FFileInfo>& InFileHashes, const TCHAR* InDestPath, bool bUseMountPoint, const FKeyChain& InKeyChain, const FString* InFilter, TArray<FPakInputPair>* OutEntries = nullptr, TArray<FPakInputPair>* OutDeletedEntries = nullptr, TMap<FString, uint64>* OutOrderMap = nullptr, TArray<FGuid>* OutUsedEncryptionKeys = nullptr, bool* OutAnyPakSigned = nullptr)
+bool ExtractFilesFromPak(const TCHAR* InPakFilename, TMap<FString, FFileInfo>& InFileHashes, const TCHAR* InDestPath, bool bUseMountPoint, const FKeyChain& InKeyChain, const FString* InFilter, TArray<FPakInputPair>* OutEntries = nullptr, TArray<FPakInputPair>* OutDeletedEntries = nullptr, FPakOrderMap* OutOrderMap = nullptr, TArray<FGuid>* OutUsedEncryptionKeys = nullptr, bool* OutAnyPakSigned = nullptr)
 {
 	// Gather all patch versions of the requested pak file and run through each separately
 	TArray<FString> PakFileList;
@@ -3037,23 +3121,287 @@ bool FileIsIdentical(FString SourceFile, FString DestFilename, const FFileInfo*
 	return true;
 }
 
-int32 CountBitToggles(const TBitArray<>& BitArray)
+float GetFragmentationPercentage(const TArray<FPakInputPair>& FilesToPak, const TBitArray<>& IncludeBitMask, int32 MaxAdjacentOrderDiff, bool bConsiderSecondaryFiles)
 {
-	int32 ChangeCount = 0;
+	uint64 PrevOrder = MAX_uint64;
 	bool bPrevBit = false;
-	for (int i = 0; i < BitArray.Num(); i++)
+	int32 DiffCount = 0;
+	int32 ConsideredCount = 0;
+	for (int32 i = 0; i < IncludeBitMask.Num(); i++)
 	{
-		bool bCurrentBit = BitArray[i];
-		if (i == 0 || bCurrentBit != bPrevBit)
+		if (!bConsiderSecondaryFiles && !FilesToPak[i].bIsInPrimaryOrder)
 		{
-			ChangeCount++;
+			PrevOrder = MAX_uint64;
+			continue;
 		}
+		uint64 CurrentOrder = FilesToPak[i].SuggestedOrder;
+		bool bCurrentBit = IncludeBitMask[i];
+		uint64 OrderDiff = CurrentOrder - PrevOrder;
+		if (OrderDiff > MaxAdjacentOrderDiff || bCurrentBit != bPrevBit)
+		{
+			DiffCount++;
+		}
+		ConsideredCount++;
+		PrevOrder = CurrentOrder;
 		bPrevBit = bCurrentBit;
 	}
-	return ChangeCount;
+	// First always shows as different, so discount it
+	DiffCount--;
+	return 100.0f * float(DiffCount) / float(ConsideredCount);
+}
+
+int64 ComputePatchSize(const TArray<FPakInputPair>& FilesToPak, const TArray<int64>& FileSizes, TBitArray<>& IncludeFilesMask, int32& OutFileCount)
+{
+	int64 TotalPatchSize = 0;
+	OutFileCount = 0;
+	for (int i = 0; i < FilesToPak.Num(); i++)
+	{
+		if (IncludeFilesMask[i])
+		{
+			OutFileCount++;
+			TotalPatchSize += FileSizes[i];
+		}
+	}
+	return TotalPatchSize;
+}
+
+int32 AddOrphanedFiles(const TArray<FPakInputPair>& FilesToPak, const TArray<int64>& FileSizes, const TArray<int32>& UAssetToUexpMapping, TBitArray<>& IncludeFilesMask, int64& OutSizeIncrease)
+{
+	int32 UExpCount = 0;
+	int32 UAssetCount = 0;
+	int32 FilesAddedCount = 0;
+	OutSizeIncrease = 0;
+	// Add corresponding UExp/UBulk files to the patch if either is included but not the other
+	for (int i = 0; i < FilesToPak.Num(); i++)
+	{
+		int32 CounterpartFileIndex = UAssetToUexpMapping[i];
+		if (CounterpartFileIndex != -1)
+		{
+			const FPakInputPair& File = FilesToPak[i];
+			const FPakInputPair& CounterpartFile = FilesToPak[CounterpartFileIndex];
+			if (IncludeFilesMask[i] != IncludeFilesMask[CounterpartFileIndex])
+			{
+				if (!IncludeFilesMask[i])
+				{
+					//UE_LOG(LogPakFile, Display, TEXT("Added %s because %s is already included"), *FilesToPak[i].Source, *FilesToPak[CounterpartFileIndex].Source);
+					IncludeFilesMask[i] = true;
+					OutSizeIncrease += FileSizes[i];
+				}
+				else
+				{
+					//UE_LOG(LogPakFile, Display, TEXT("Added %s because %s is already included"), *FilesToPak[CounterpartFileIndex].Source, *FilesToPak[i].Source);
+					IncludeFilesMask[CounterpartFileIndex] = true;
+					OutSizeIncrease += FileSizes[CounterpartFileIndex];
+				}
+				FilesAddedCount++;
+			}
+		}
+	}
+	return FilesAddedCount;
+}
+
+
+bool DoGapFillingIteration(const TArray<FPakInputPair>& FilesToPak, const TArray<int64>& FileSizes, const TArray<int32>& UAssetToUexpMapping, TBitArray<>& InOutIncludeFilesMask, int64 MaxGapSizeBytes, int32 MaxAdjacentOrderDiff, bool bForce, int64 MaxPatchSize = 0, bool bFillPrimaryOrderFiles = true, bool bFillSecondaryOrderFiles = true)
+{
+	TBitArray<> IncludeFilesMask = InOutIncludeFilesMask;
+
+	float FragmentationPercentageOriginal = GetFragmentationPercentage(FilesToPak, IncludeFilesMask, MaxAdjacentOrderDiff, true);
+
+	int64 SizeIncrease = 0;
+	int64 CurrentOffset = 0;
+	int64 CurrentPatchOffset = 0;
+	int64 CurrentGapSize = 0;
+	bool bPrevKeepFile = false;
+	uint64 PrevOrder = MAX_uint64;
+	int32 OriginalKeepCount = 0;
+	int32 LastKeepIndex = -1;
+	bool bCurrentGapIsUnbroken = true;
+	int32 PatchFilesAddedCount = 0;
+	for (int i = 0; i < FilesToPak.Num(); i++)
+	{
+		bool bKeepFile = IncludeFilesMask[i];
+		const FPakInputPair& File = FilesToPak[i];
+		uint64 Order = File.SuggestedOrder;
+
+		// Skip unordered files or files outside the range we care about
+		if (Order == MAX_uint64 || (File.bIsInPrimaryOrder && !bFillPrimaryOrderFiles) || (!File.bIsInPrimaryOrder && !bFillSecondaryOrderFiles))
+		{
+			continue;
+		}
+		CurrentOffset += FileSizes[i];
+		if (bKeepFile)
+		{
+			OriginalKeepCount++;
+			CurrentPatchOffset = CurrentOffset;
+		}
+		else
+		{
+			if (OriginalKeepCount > 0)
+			{
+				CurrentGapSize = CurrentOffset - CurrentPatchOffset;
+			}
+		}
+
+		// Detect gaps in the file order. No point in removing those gaps because it won't affect seeks
+		uint64 OrderDiff = Order - PrevOrder;
+		if (bCurrentGapIsUnbroken && OrderDiff > uint64(MaxAdjacentOrderDiff))
+		{
+			bCurrentGapIsUnbroken = false;
+		}
+
+		// If we're keeping this file but not the last one, check if the gap size is small enough to bring over unchanged assets
+		if (bKeepFile && !bPrevKeepFile && CurrentGapSize > 0)
+		{
+			if (CurrentGapSize <= MaxGapSizeBytes)
+			{
+				if (bCurrentGapIsUnbroken)
+				{
+					// Mark the files in the gap to keep, even though they're unchanged
+					for (int j = LastKeepIndex + 1; j < i; j++)
+					{
+						IncludeFilesMask[j] = true;
+						SizeIncrease += FileSizes[j];
+						PatchFilesAddedCount++;
+					}
+				}
+			}
+			bCurrentGapIsUnbroken = true;
+		}
+		bPrevKeepFile = bKeepFile;
+		if (bKeepFile)
+		{
+			LastKeepIndex = i;
+		}
+		PrevOrder = Order;
+	}
+
+#if GUARANTEE_UASSET_AND_UEXP_IN_SAME_PAK
+	int64 OrphanedFilesSizeIncrease = 0;
+	int32 OrphanedFileCount = AddOrphanedFiles(FilesToPak, FileSizes, UAssetToUexpMapping, IncludeFilesMask, OrphanedFilesSizeIncrease);
+#endif
+
+	int32 PatchFileCount = 0;
+	int64 PatchSize = ComputePatchSize(FilesToPak, FileSizes, IncludeFilesMask, PatchFileCount);
+
+	FString PrefixString = TEXT("");
+	if (bFillPrimaryOrderFiles && !bFillSecondaryOrderFiles)
+	{
+		PrefixString = TEXT("[PRIMARY]");
+	}
+	else if (bFillSecondaryOrderFiles && !bFillPrimaryOrderFiles)
+	{
+		PrefixString = TEXT("[SECONDARY]");
+	}
+	if (PatchSize > MaxPatchSize && !bForce)
+	{
+		UE_LOG(LogPakFile, SEEK_OPT_VERBOSITY, TEXT("Patch seek optimization step %s - Gap size: %dKB - FAILED (patch too big)"), *PrefixString, MaxGapSizeBytes / 1024);
+		return false;
+	}
+
+	// Stop if we didn't actually make patch size better
+	float FragmentationPercentageAfterGapFill = GetFragmentationPercentage(FilesToPak, IncludeFilesMask, MaxAdjacentOrderDiff, true);
+	if (FragmentationPercentageAfterGapFill >= FragmentationPercentageOriginal && !bForce)
+	{
+		UE_LOG(LogPakFile, SEEK_OPT_VERBOSITY, TEXT("Patch seek optimization step %s - Gap size: %dKB - FAILED (contiguous block count didn't improve)"), *PrefixString, MaxGapSizeBytes / 1024);
+		return false;
+	}
+	InOutIncludeFilesMask = IncludeFilesMask;
+	UE_LOG(LogPakFile, SEEK_OPT_VERBOSITY, TEXT("Patch seek optimization step %s - Gap size: %dKB - SUCCEEDED"), *PrefixString, MaxGapSizeBytes / 1024);
+	return true;
 }
 
-void RemoveIdenticalFiles( TArray<FPakInputPair>& FilesToPak, const FString& SourceDirectory, const TMap<FString, FFileInfo>& FileHashes, int64 SeekOptMaxGapSizeBytes, bool bSeekOptUseOrder)
+bool DoIncrementalGapFilling(const TArray<FPakInputPair>& FilesToPak, const TArray<int64>& FileSizes, const TArray<int32>& UAssetToUexpMapping, TBitArray<>& IncludeFilesMask, int64 MinGapSize, int64 MaxGapSize, int64 MaxAdjacentOrderDiff, bool bFillPrimaryOrderFiles, bool bFillSecondaryOrderFiles, int64 MaxPatchSize)
+{
+	int64 GapSize = MinGapSize;
+	bool bSuccess = false;
+	while (GapSize <= MaxGapSize)
+	{
+		if (DoGapFillingIteration(FilesToPak, FileSizes, UAssetToUexpMapping, IncludeFilesMask, GapSize, MaxAdjacentOrderDiff, false, MaxPatchSize, bFillPrimaryOrderFiles, bFillSecondaryOrderFiles))
+		{
+			bSuccess = true;
+		}
+		else
+		{
+			// Try with 75% of the max gap size
+			bSuccess |= DoGapFillingIteration(FilesToPak, FileSizes, UAssetToUexpMapping, IncludeFilesMask, GapSize*0.75, MaxAdjacentOrderDiff, false, MaxPatchSize, bFillPrimaryOrderFiles, bFillSecondaryOrderFiles);
+			break;
+		}
+		GapSize *= 2;
+	}
+	return bSuccess;
+}
+
+void ApplyGapFilling(const TArray<FPakInputPair>& FilesToPak, const TArray<int64>& FileSizes, const TArray<int32>& UAssetToUexpMapping, TBitArray<>& IncludeFilesMask, const FPatchSeekOptParams& SeekOptParams)
+{
+	UE_CLOG(SeekOptParams.MaxGapSize == 0, LogPakFile, Fatal, TEXT("ApplyGapFilling requires MaxGapSize > 0"));
+	check(SeekOptParams.MaxGapSize > 0);
+	float FragmentationPercentageOriginal = GetFragmentationPercentage(FilesToPak, IncludeFilesMask, SeekOptParams.MaxAdjacentOrderDiff, true);
+	float FragmentationPercentageOriginalPrimary = GetFragmentationPercentage(FilesToPak, IncludeFilesMask, SeekOptParams.MaxAdjacentOrderDiff, false);
+	int32 OriginalPatchFileCount = 0;
+	int64 OriginalPatchSize = ComputePatchSize(FilesToPak, FileSizes, IncludeFilesMask, OriginalPatchFileCount);
+	int64 IncrementalMaxPatchSize = int64(double(OriginalPatchSize) + double(OriginalPatchSize) * double(SeekOptParams.MaxInflationPercent) * 0.01);
+	int64 MinIncrementalGapSize = 4 * 1024;
+
+	ESeekOptMode SeekOptMode = SeekOptParams.Mode;
+	switch (SeekOptMode)
+	{
+	case ESeekOptMode::OnePass:
+	{
+		DoGapFillingIteration(FilesToPak, FileSizes, UAssetToUexpMapping, IncludeFilesMask, SeekOptParams.MaxGapSize, SeekOptParams.MaxAdjacentOrderDiff, true);
+	}
+	break;
+	case ESeekOptMode::Incremental:
+	case ESeekOptMode::Incremental_OnlyPrimaryOrder:
+	{
+		UE_CLOG(SeekOptParams.MaxInflationPercent == 0.0f, LogPakFile, Fatal, TEXT("ESeekOptMode::Incremental* requires MaxInflationPercent > 0.0"));
+		bool bFillSecondaryOrderFiles = (SeekOptParams.Mode == ESeekOptMode::Incremental);
+		DoIncrementalGapFilling(FilesToPak, FileSizes, UAssetToUexpMapping, IncludeFilesMask, MinIncrementalGapSize, SeekOptParams.MaxGapSize, SeekOptParams.MaxAdjacentOrderDiff, true, bFillSecondaryOrderFiles, IncrementalMaxPatchSize);
+	}
+	break;
+	case ESeekOptMode::Incremental_PrimaryThenSecondary:
+	{
+		UE_CLOG(SeekOptParams.MaxInflationPercent == 0.0f, LogPakFile, Fatal, TEXT("ESeekOptMode::Incremental* requires MaxInflationPercent > 0.0"));
+		int64 PassMaxPatchSize[3];
+		PassMaxPatchSize[0] = OriginalPatchSize + (IncrementalMaxPatchSize - OriginalPatchSize) * 0.9f;
+		PassMaxPatchSize[1] = IncrementalMaxPatchSize;
+		PassMaxPatchSize[2] = IncrementalMaxPatchSize;
+		for (int32 i = 0; i < 3; i++)
+		{
+			bool bFillPrimaryOrderFiles = (i == 0) || (i == 2);
+			bool bFillSecondaryOrderFiles = (i == 1);
+			DoIncrementalGapFilling(FilesToPak, FileSizes, UAssetToUexpMapping, IncludeFilesMask, MinIncrementalGapSize, SeekOptParams.MaxGapSize, SeekOptParams.MaxAdjacentOrderDiff, bFillPrimaryOrderFiles, bFillSecondaryOrderFiles, PassMaxPatchSize[i]);
+		}
+	}
+	break;
+	}
+
+	int32 NewPatchFileCount = 0;
+	int64 NewPatchSize = ComputePatchSize(FilesToPak, FileSizes, IncludeFilesMask, NewPatchFileCount);
+
+	double OriginalSizeMB = double(OriginalPatchSize) / 1024.0 / 1024.0;
+	double SizeIncreaseMB = double(NewPatchSize - OriginalPatchSize) / 1024.0 / 1024.0;
+	double TotalSizeMB = OriginalSizeMB + SizeIncreaseMB;
+	double SizeIncreasePercent = 100.0 * SizeIncreaseMB / OriginalSizeMB;
+	if (NewPatchFileCount == OriginalPatchSize)
+	{
+		UE_LOG(LogPakFile, SEEK_OPT_VERBOSITY, TEXT("Patch seek optimization did not modify patch pak size (no additional files added)"), OriginalSizeMB, TotalSizeMB, SizeIncreasePercent);
+	}
+	else
+	{
+		UE_LOG(LogPakFile, SEEK_OPT_VERBOSITY, TEXT("Patch seek optimization increased estimated patch pak size from %.2fMB to %.2fMB (+%.1f%%)"), OriginalSizeMB, TotalSizeMB, SizeIncreasePercent);
+		UE_LOG(LogPakFile, SEEK_OPT_VERBOSITY, TEXT("Total files added : %d (from %d)"), NewPatchFileCount - OriginalPatchFileCount, OriginalPatchFileCount);
+	}
+
+
+	float FragmentationPercentageNew = GetFragmentationPercentage(FilesToPak, IncludeFilesMask, SeekOptParams.MaxAdjacentOrderDiff, true);
+	float FragmentationPercentageNewPrimary = GetFragmentationPercentage(FilesToPak, IncludeFilesMask, SeekOptParams.MaxAdjacentOrderDiff, false);
+	UE_LOG(LogPakFile, SEEK_OPT_VERBOSITY, TEXT("Fragmentation pre-optimization: %.2f%%"), FragmentationPercentageOriginal);
+	UE_LOG(LogPakFile, SEEK_OPT_VERBOSITY, TEXT("Fragmentation final: %.2f%%"), FragmentationPercentageNew);
+	UE_LOG(LogPakFile, SEEK_OPT_VERBOSITY, TEXT("Fragmentation pre-optimization (primary files): %.2f%%"), FragmentationPercentageOriginalPrimary);
+	UE_LOG(LogPakFile, SEEK_OPT_VERBOSITY, TEXT("Fragmentation final (primary files): %.2f%%"), FragmentationPercentageNewPrimary);
+}
+
+void RemoveIdenticalFiles( TArray<FPakInputPair>& FilesToPak, const FString& SourceDirectory, const TMap<FString, FFileInfo>& FileHashes, const FPatchSeekOptParams& SeekOptParams)
 {
 	FString HashFilename = SourceDirectory / TEXT("Hashes.txt");
 
@@ -3119,14 +3467,18 @@ void RemoveIdenticalFiles( TArray<FPakInputPair>& FilesToPak, const FString& Sou
  			UE_LOG(LogPakFile, Display, TEXT("Didn't find hash for %s No mount %s"), *SourceFilename, *SourceFileNoMountPoint);
  		}
  
+#if GUARANTEE_UASSET_AND_UEXP_IN_SAME_PAK
 		// uexp files are always handled with their corresponding uasset file
 		bool bForceInclude = false;
 		if (!FPaths::GetExtension(SourceFilename).Equals("uexp", ESearchCase::IgnoreCase))
 		{
 			bForceInclude = FoundFileHash && FoundFileHash->bForceInclude;
 		}
+#else
+		bool bForceInclude = FoundFileHash && FoundFileHash->bForceInclude;
+#endif
 
-			FString DestFilename = NewFile.Source;
+		FString DestFilename = NewFile.Source;
 		if (!bForceInclude && FileIsIdentical(SourceFilename, DestFilename, FoundFileHash, &FileSizes[i]))
 		{
 			UE_LOG(LogPakFile, Display, TEXT("Source file %s matches dest file %s and will not be included in patch"), *SourceFilename, *DestFilename);
@@ -3136,6 +3488,7 @@ void RemoveIdenticalFiles( TArray<FPakInputPair>& FilesToPak, const FString& Sou
 	}
 
 	// Add corresponding UExp/UBulk files to the patch if one is included but not the other (uassets and uexp files must be in the same pak)
+#if GUARANTEE_UASSET_AND_UEXP_IN_SAME_PAK
 	for (int i = 0; i < FilesToPak.Num(); i++)
 	{
 		int32 CounterpartFileIndex= UAssetToUexpMapping[i];
@@ -3149,126 +3502,11 @@ void RemoveIdenticalFiles( TArray<FPakInputPair>& FilesToPak, const FString& Sou
 			}
 		}
 	}
+#endif
 
-	if (SeekOptMaxGapSizeBytes > 0)
-				{
-		UE_LOG(LogPakFile, Display, TEXT("Patch seek optimization - filling gaps up to %dKB"), SeekOptMaxGapSizeBytes/1024);
-
-		int32 PatchContiguousBlockCountOriginal = CountBitToggles(IncludeFilesMask);
-		int64 MaxGapSizeBytes = 0;
-		int64 OriginalPatchSize = 0;
-		int64 SizeIncrease = 0;
-		int64 CurrentOffset = 0;
-		int64 CurrentPatchOffset = 0;
-		int64 CurrentGapSize = 0;
-		bool bPrevKeepFile = false;
-		uint64 PrevOrder=MAX_uint64;
-		int32 OriginalKeepCount = 0;
-		int32 LastKeepIndex = -1;
-		bool bCurrentGapIsUnbroken = true;
-		int32 PatchFilesAddedCount = 0;
-		int32 OriginalPatchFileCount = 0;
-		for (int i = 0; i < FilesToPak.Num(); i++)
-					{
-			bool bKeepFile = IncludeFilesMask[i];
-			uint64 Order = FilesToPak[i].SuggestedOrder;
-			if (bKeepFile)
-			{
-				OriginalPatchFileCount++;
-				OriginalPatchSize += FileSizes[i];
-					}
-
-			if (Order == MAX_uint64)
-			{
-				// Skip unordered files
-				continue;
-			}
-			CurrentOffset += FileSizes[i];
-			if (bKeepFile)
-					{
-				OriginalKeepCount++;
-				CurrentPatchOffset = CurrentOffset;
-			}
-			else
-			{
-				if (OriginalKeepCount > 0)
-				{
-					CurrentGapSize = CurrentOffset - CurrentPatchOffset;
-				}
-					}
-
-			// Detect gaps in the file order. No point in removing those gaps because it won't affect seeks
-			if (bCurrentGapIsUnbroken && Order != PrevOrder + 1)
-					{
-				bCurrentGapIsUnbroken = false;
-			}
-
-			// If we're keeping this file but not the last one, check if the gap size is small enough to bring over unchanged assets
-			if (bKeepFile && !bPrevKeepFile && CurrentGapSize > 0 )
-			{
-				if ( CurrentGapSize <= SeekOptMaxGapSizeBytes)
-				{
-					if (bCurrentGapIsUnbroken || bSeekOptUseOrder==false)
-					{
-						// Mark the files in the gap to keep, even though they're unchanged
-						for (int j = LastKeepIndex + 1; j < i; j++)
-						{
-							IncludeFilesMask[j] = true;
-							SizeIncrease += FileSizes[j];
-							PatchFilesAddedCount++;
-						}
-					}
-				}
-				bCurrentGapIsUnbroken = true;
-			}
-			bPrevKeepFile = bKeepFile;
-			if (bKeepFile)
-			{
-				LastKeepIndex = i;
-			}
-			PrevOrder = Order;
-		}
-
-		// Add corresponding UExp/UBulk files to the patch if either is included but not the other
-		for (int i = 0; i < FilesToPak.Num(); i++)
-		{
-			int32 CounterpartFileIndex = UAssetToUexpMapping[i];
-			if (CounterpartFileIndex != -1)
-			{
-				if (IncludeFilesMask[i] != IncludeFilesMask[CounterpartFileIndex])
-				{
-					if ( !IncludeFilesMask[i] )
-					{
-						IncludeFilesMask[i] = true;
-						SizeIncrease += FileSizes[i];
-					}
-					else
-						{
-						IncludeFilesMask[CounterpartFileIndex] = true;
-						SizeIncrease += FileSizes[CounterpartFileIndex];
-					}
-					PatchFilesAddedCount++;
-						}
-					}
-				}
-
-		double OriginalSizeMB = double(OriginalPatchSize) / 1024.0 / 1024.0;
-		double SizeIncreaseMB = double(SizeIncrease) / 1024.0 / 1024.0;
-		double TotalSizeMB = OriginalSizeMB + SizeIncreaseMB;
-		double SizeIncreasePercent = 100.0 * SizeIncreaseMB / OriginalSizeMB;
-		if (PatchFilesAddedCount == 0)
-		{
-			UE_LOG(LogPakFile, Display, TEXT("Patch seek optimization did not modify patch pak size (no additional files added)"), OriginalSizeMB, TotalSizeMB, SizeIncreasePercent);
-			}
-		else
-		{
-			UE_LOG(LogPakFile, Display, TEXT("Patch seek optimization increased estimated patch pak size from %.2fMB to %.2fMB (+%.1f%%)"), OriginalSizeMB, TotalSizeMB, SizeIncreasePercent);
-			UE_LOG(LogPakFile, Display, TEXT("Total files added : %d (of %d)"), PatchFilesAddedCount, OriginalPatchFileCount + PatchFilesAddedCount);
-		}
-		UE_LOG(LogPakFile, Display, TEXT("Contiguous block count pre-optimization: %d"), PatchContiguousBlockCountOriginal);
-
-		int32 PatchContiguousBlockCountFinal = CountBitToggles(IncludeFilesMask);
-		UE_LOG(LogPakFile, Display, TEXT("Contiguous block count final: %d"), PatchContiguousBlockCountFinal);
+	if (SeekOptParams.Mode != ESeekOptMode::None)
+	{
+		ApplyGapFilling(FilesToPak, FileSizes, UAssetToUexpMapping, IncludeFilesMask, SeekOptParams);
 	}
 
 	// Compress the array while preserving the order, removing the files we marked to remove
@@ -3492,7 +3730,7 @@ bool Repack(const FString& InputPakFile, const FString& OutputPakFile, const FPa
 	TMap<FString, FFileInfo> Hashes;
 	TArray<FPakInputPair> Entries;
 	TArray<FPakInputPair> DeletedEntries;
-	TMap<FString, uint64> OrderMap;
+	FPakOrderMap OrderMap;
 	TArray<FGuid> EncryptionKeys;
 	bool bAnySigned = false;
 	FString TempDir = FPaths::EngineIntermediateDir() / TEXT("UnrealPak") / TEXT("Repack") / FPaths::GetBaseFilename(InputPakFile);
@@ -3696,14 +3934,14 @@ bool ExecuteUnrealPak(const TCHAR* CmdLine)
 		bool bOnlyDeleted = FParse::Param( CmdLine, TEXT("OnlyDeleted") );
 		bool bSortByOrdering = FParse::Param(CmdLine, TEXT("SortByOrdering"));
 
-		TMap<FString, uint64> OrderMap;
+		FPakOrderMap OrderMap;
 		FString ResponseFile;
-		if (FParse::Value(CmdLine, TEXT("-order="), ResponseFile) && !ProcessOrderFile(*ResponseFile, OrderMap))
+		if (FParse::Value(CmdLine, TEXT("-order="), ResponseFile) && !OrderMap.ProcessOrderFile(*ResponseFile))
 		{
 			return false;
 		}
 		FString SecondaryResponseFile;
-		if (FParse::Value(CmdLine, TEXT("-secondaryOrder="), SecondaryResponseFile) && !ProcessOrderFile(*SecondaryResponseFile, OrderMap, true, OrderMap.Num()))
+		if (FParse::Value(CmdLine, TEXT("-secondaryOrder="), SecondaryResponseFile) && !OrderMap.ProcessOrderFile(*SecondaryResponseFile, true))
 		{
 			return false;
 		}
@@ -3847,15 +4085,15 @@ bool ExecuteUnrealPak(const TCHAR* CmdLine)
 		FPakCommandLineParameters CmdLineParameters;
 		ProcessCommandLine(CmdLine, NonOptionArguments, Entries, CmdLineParameters);
 
-		TMap<FString, uint64> OrderMap;
+		FPakOrderMap OrderMap;
 		FString ResponseFile;
-		if (FParse::Value(CmdLine, TEXT("-order="), ResponseFile) && !ProcessOrderFile(*ResponseFile, OrderMap))
+		if (FParse::Value(CmdLine, TEXT("-order="), ResponseFile) && !OrderMap.ProcessOrderFile(*ResponseFile))
 		{
 			return false;
 		}
 
 		FString SecondaryResponseFile;
-		if (FParse::Value(CmdLine, TEXT("-secondaryOrder="), SecondaryResponseFile) && !ProcessOrderFile(*SecondaryResponseFile, OrderMap, true, OrderMap.Num()))
+		if (FParse::Value(CmdLine, TEXT("-secondaryOrder="), SecondaryResponseFile) && !OrderMap.ProcessOrderFile(*SecondaryResponseFile, true))
 		{
 			return false;
 		}
@@ -3931,7 +4169,7 @@ bool ExecuteUnrealPak(const TCHAR* CmdLine)
 			FilesToAdd.Append(DeleteRecords);
 
 			// if we are generating a patch here we remove files which are already shipped...
-			RemoveIdenticalFiles(FilesToAdd, CmdLineParameters.SourcePatchDiffDirectory, SourceFileHashes, CmdLineParameters.PatchSeekOptMaxGapSize, CmdLineParameters.PatchSeekOptUseOrder);
+			RemoveIdenticalFiles(FilesToAdd, CmdLineParameters.SourcePatchDiffDirectory, SourceFileHashes, CmdLineParameters.SeekOptParams);
 		}
 
 
diff --git a/Engine/Source/Programs/AutomationTool/Scripts/CopyBuildToStagingDirectory.Automation.cs b/Engine/Source/Programs/AutomationTool/Scripts/CopyBuildToStagingDirectory.Automation.cs
--- a/Engine/Source/Programs/AutomationTool/Scripts/CopyBuildToStagingDirectory.Automation.cs
+++ b/Engine/Source/Programs/AutomationTool/Scripts/CopyBuildToStagingDirectory.Automation.cs
@@ -1948,14 +1948,32 @@ public partial class Project : CommandUtils
                         }
 					}
 
-					string PatchSeekOptMaxGapSizeOption = String.Empty;
-					string PatchSeekOptMaxGapSize = String.Empty;
-					if (PlatformGameConfig.GetString("/Script/UnrealEd.ProjectPackagingSettings", "PatchSeekOptMaxGapSize", out PatchSeekOptMaxGapSize))
+					string PatchSeekOptArgs = String.Empty;
+					if (bShouldGeneratePatch)
 					{
-						PatchSeekOptMaxGapSizeOption = String.Format(" -patchSeekOptMaxGapSize={0}", PatchSeekOptMaxGapSize);
+						string PatchSeekOptMode = String.Empty;
+						string PatchSeekOptMaxGapSize = String.Empty;
+						string PatchSeekOptMaxInflationPercent = String.Empty;
+						string PatchSeekOptMaxAdjacentOrderDiff = String.Empty;
+						if (PlatformGameConfig.GetString("/Script/UnrealEd.ProjectPackagingSettings", "PatchSeekOptMode", out PatchSeekOptMode))
+						{
+							PatchSeekOptArgs += String.Format(" -PatchSeekOptMode={0}", PatchSeekOptMode);
+						}
+						if (PlatformGameConfig.GetString("/Script/UnrealEd.ProjectPackagingSettings", "PatchSeekOptMaxGapSize", out PatchSeekOptMaxGapSize))
+						{
+							PatchSeekOptArgs += String.Format(" -PatchSeekOptMaxGapSize={0}", PatchSeekOptMaxGapSize);
+						}
+						if (PlatformGameConfig.GetString("/Script/UnrealEd.ProjectPackagingSettings", "PatchSeekOptMaxInflationPercent", out PatchSeekOptMaxInflationPercent))
+						{
+							PatchSeekOptArgs += String.Format(" -PatchSeekOptMaxInflationPercent={0}", PatchSeekOptMaxInflationPercent);
+						}
+						if (PlatformGameConfig.GetString("/Script/UnrealEd.ProjectPackagingSettings", "PatchSeekOptMaxAdjacentOrderDiff", out PatchSeekOptMaxAdjacentOrderDiff))
+						{
+							PatchSeekOptArgs += String.Format(" -PatchSeekOptMaxAdjacentOrderDiff={0}", PatchSeekOptMaxAdjacentOrderDiff);
+						}
 					}
 
-					Commands.Add(GetUnrealPakArguments(Params.RawProjectPath, PakParams.UnrealPakResponseFile, OutputLocation, PrimaryOrderFile, SC.StageTargetPlatform.GetPlatformPakCommandLine(Params, SC) + PatchSeekOptMaxGapSizeOption + BulkOption + CompressionFormats + " " + Params.AdditionalPakOptions, PakParams.bCompressed, CryptoSettings, CryptoKeysCacheFilename, PatchSourceContentPath, PakParams.EncryptionKeyGuid, SecondaryOrderFile));
+					Commands.Add(GetUnrealPakArguments(Params.RawProjectPath, PakParams.UnrealPakResponseFile, OutputLocation, PrimaryOrderFile, SC.StageTargetPlatform.GetPlatformPakCommandLine(Params, SC) + PatchSeekOptArgs + BulkOption + CompressionFormats + " " + Params.AdditionalPakOptions, PakParams.bCompressed, CryptoSettings, CryptoKeysCacheFilename, PatchSourceContentPath, PakParams.EncryptionKeyGuid, SecondaryOrderFile));
 					LogNames.Add(OutputLocation.GetFileNameWithoutExtension());
 				}
 			}
