Added additional logging to unrealpak for compressor settings.
#test Unrealpak

diff --git a/Engine/Plugins/Runtime/PacketHandlers/CompressionComponents/Oodle/Source/OodleCompressionFormat/Private/OodleCompressionFormat.cpp b/Engine/Plugins/Runtime/PacketHandlers/CompressionComponents/Oodle/Source/OodleCompressionFormat/Private/OodleCompressionFormat.cpp
--- a/Engine/Plugins/Runtime/PacketHandlers/CompressionComponents/Oodle/Source/OodleCompressionFormat/Private/OodleCompressionFormat.cpp
+++ b/Engine/Plugins/Runtime/PacketHandlers/CompressionComponents/Oodle/Source/OodleCompressionFormat/Private/OodleCompressionFormat.cpp
@@ -130,7 +130,7 @@ class FOodleCompressionFormatModuleInterface : public IModuleInterface
 		OodleLZ_Compressor UsedCompressor = MethodMap.FindRef(Method);
 		OodleLZ_CompressionLevel UsedLevel = LevelMap.FindRef(Level);
 
-		UE_LOG(OodleCompression, Display, TEXT("Oodle Compressing with %s, level %s, SpaceSpeed tradeoff %d"), **MethodMap.FindKey(UsedCompressor), **LevelMap.FindKey(UsedLevel), SpaceSpeedTradeoff, *GEngineIni);
+		UE_LOG(OodleCompression, Display, TEXT("Oodle initializing compressor with %s, level %s, SpaceSpeed tradeoff %d"), **MethodMap.FindKey(UsedCompressor), **LevelMap.FindKey(UsedLevel), SpaceSpeedTradeoff, *GEngineIni);
 
 		CompressionFormat = new FOodleCustomCompressor(MethodMap.FindRef(Method), LevelMap.FindRef(Level), SpaceSpeedTradeoff);
 
diff --git a/Engine/Source/Developer/PakFileUtilities/Private/PakFileUtilities.cpp b/Engine/Source/Developer/PakFileUtilities/Private/PakFileUtilities.cpp
--- a/Engine/Source/Developer/PakFileUtilities/Private/PakFileUtilities.cpp
+++ b/Engine/Source/Developer/PakFileUtilities/Private/PakFileUtilities.cpp
@@ -1529,6 +1529,7 @@ bool CreatePakFile(const TCHAR* Filename, TArray<FPakInputPair>& FilesToAdd, con
 
 	}
 
+	TArray<FName> UsedCompressionFormats; // List of compression formats we actually used in this pak file (used for logging only)
 
 	if (InKeyChain.MasterEncryptionKey)
 	{
@@ -1810,6 +1811,7 @@ bool CreatePakFile(const TCHAR* Filename, TArray<FPakInputPair>& FilesToAdd, con
 			}
 			if (FilesToAdd[FileIndex].bNeedsCompression && CompressionMethod != NAME_None)
 			{
+				UsedCompressionFormats.AddUnique(CompressionMethod); // used for logging only
 				FinalizeCopyCompressedFileToPak(Info, CompressedFileBuffer, NewEntry);
 			}
 
@@ -1948,6 +1950,14 @@ bool CreatePakFile(const TCHAR* Filename, TArray<FPakInputPair>& FilesToAdd, con
 	{
 		float PercentLess = ((float)TotalCompressedSize / (TotalUncompressedSize / 100.f));
 		UE_LOG(LogPakFile, Display, TEXT("Compression summary: %.2f%% of original size. Compressed Size %lld bytes, Original Size %lld bytes. "), PercentLess, TotalCompressedSize, TotalUncompressedSize);
+
+		FString UsedCompressionFormatsString;
+		for (FName CompressionFormat : UsedCompressionFormats)
+		{
+			UsedCompressionFormatsString.Append( CompressionFormat.ToString() + TEXT(", ") );
+		}
+
+		UE_LOG(LogPakFile, Display, TEXT("Used compression formats (in priority order) '%s'"), *UsedCompressionFormatsString);
 	}
 
 	if (TotalEncryptedDataSize)
