Fix bug caught by static analysis in AuditPakFiles()

diff --git a/Engine/Source/Developer/PakFileUtilities/Private/PakFileUtilities.cpp b/Engine/Source/Developer/PakFileUtilities/Private/PakFileUtilities.cpp
--- a/Engine/Source/Developer/PakFileUtilities/Private/PakFileUtilities.cpp
+++ b/Engine/Source/Developer/PakFileUtilities/Private/PakFileUtilities.cpp
@@ -2331,7 +2331,8 @@ bool AuditPakFiles( const FString& InputPath, bool bOnlyDeleted, const FString&
 			FPaths::NormalizeFilename(OpenOrderAssetName);
 			OpenOrderAssetName.ToLowerInline();
 
-			if( uint64 OrderIndex = OrderMap.GetFileOrder( OpenOrderAssetName, false ) != MAX_uint64 )
+			uint64 OrderIndex = OrderMap.GetFileOrder(OpenOrderAssetName, false);
+			if (OrderIndex != MAX_uint64)
 			{
 				CachedOpenOrder.Add( AssetPath, OrderIndex );
 			}
