Add -fallbackOrderForNonUassetFiles option for ordering uexp and ubulk files if order is not specified (e.g for the secondary/cooker open order). Previously this was non-optional, but this caused a regression in patch sizes on some platforms

diff --git a/Engine/Source/Developer/PakFileUtilities/Private/PakFileUtilities.cpp b/Engine/Source/Developer/PakFileUtilities/Private/PakFileUtilities.cpp
--- a/Engine/Source/Developer/PakFileUtilities/Private/PakFileUtilities.cpp
+++ b/Engine/Source/Developer/PakFileUtilities/Private/PakFileUtilities.cpp
@@ -363,6 +363,7 @@ struct FPakCommandLineParameters
 		, UseCustomCompressor(false)
 		, bSign(false)
 		, bPatchCompatibilityMode421(false)
+		, bFallbackOrderForNonUassetFiles(false)
 	{
 	}
 
@@ -380,6 +381,7 @@ struct FPakCommandLineParameters
 	FGuid EncryptionKeyGuid;
 	bool bSign;
 	bool bPatchCompatibilityMode421;
+	bool bFallbackOrderForNonUassetFiles;
 };
 
 struct FPakEntryPair
@@ -762,7 +764,12 @@ void ProcessCommandLine(const TCHAR* CmdLine, const TArray<FString>& NonOptionAr
 	if (FParse::Param(CmdLine, TEXT("patchcompatibilitymode421")))
 	{
 		CmdLineParameters.bPatchCompatibilityMode421 = true;
-		}
+	}
+
+	if (FParse::Param(CmdLine, TEXT("fallbackOrderForNonUassetFiles")))
+	{
+		CmdLineParameters.bFallbackOrderForNonUassetFiles = true;
+	}
 
 	if (FParse::Value(CmdLine, TEXT("-blocksize="), ClusterSizeString) && 
 		FParse::Value(CmdLine, TEXT("-blocksize="), CmdLineParameters.FileSystemBlockSize))
@@ -976,7 +983,7 @@ void ProcessCommandLine(const TCHAR* CmdLine, const TArray<FString>& NonOptionAr
 	UE_LOG(LogPakFile, Display, TEXT("Added %d entries to add to pak file."), Entries.Num());
 }
 
-void CollectFilesToAdd(TArray<FPakInputPair>& OutFilesToAdd, const TArray<FPakInputPair>& InEntries, const FPakOrderMap& OrderMap)
+void CollectFilesToAdd(TArray<FPakInputPair>& OutFilesToAdd, const TArray<FPakInputPair>& InEntries, const FPakOrderMap& OrderMap, const FPakCommandLineParameters& CmdLineParameters)
 {
 	UE_LOG(LogPakFile, Display, TEXT("Collecting files to add to pak file..."));
 	const double StartTime = FPlatformTime::Seconds();
@@ -1064,7 +1071,7 @@ void CollectFilesToAdd(TArray<FPakInputPair>& OutFilesToAdd, const TArray<FPakIn
 			FileInput.Source = Input.Source;
 			FPaths::MakeStandardFilename(FileInput.Source);
 			FileInput.Dest = FileInput.Source.Replace(*Directory, *Input.Dest, ESearchCase::IgnoreCase);
-			uint64 FileOrder = OrderMap.GetFileOrder(FileInput.Dest, true, &FileInput.bIsInPrimaryOrder);
+			uint64 FileOrder = OrderMap.GetFileOrder(FileInput.Dest, CmdLineParameters.bFallbackOrderForNonUassetFiles, &FileInput.bIsInPrimaryOrder);
 			if (FileOrder != MAX_uint64)
 			{
 				FileInput.SuggestedOrder = FileOrder;
@@ -3748,7 +3755,7 @@ bool Repack(const FString& InputPakFile, const FString& OutputPakFile, const FPa
 	if (ExtractFilesFromPak(*InputPakFile, Hashes, *TempDir, false, InKeyChain, nullptr, &Entries, &DeletedEntries, &OrderMap, &EncryptionKeys, &bAnySigned))
 	{
 		TArray<FPakInputPair> FilesToAdd;
-		CollectFilesToAdd(FilesToAdd, Entries, OrderMap);
+		CollectFilesToAdd(FilesToAdd, Entries, OrderMap, CmdLineParameters);
 
 		if (bIncludeDeleted)
 		{
@@ -4159,7 +4166,7 @@ bool ExecuteUnrealPak(const TCHAR* CmdLine)
 
 		// Start collecting files
 		TArray<FPakInputPair> FilesToAdd;
-		CollectFilesToAdd(FilesToAdd, Entries, OrderMap);
+		CollectFilesToAdd(FilesToAdd, Entries, OrderMap, CmdLineParameters);
 
 		if ( CmdLineParameters.GeneratePatch )
 		{
@@ -4223,5 +4230,7 @@ bool ExecuteUnrealPak(const TCHAR* CmdLine)
 	UE_LOG(LogPakFile, Error, TEXT("    -compressionformat[s]=<Format[,format2,...]> (set the format(s) to compress with, falling back on failures)"));
 	UE_LOG(LogPakFile, Error, TEXT("    -encryptionkeyoverrideguid (override the encryption key guid used for encrypting data in this pak file)"));
 	UE_LOG(LogPakFile, Error, TEXT("    -sign (generate a signature (.sig) file alongside the pak)"));
+	UE_LOG(LogPakFile, Error, TEXT("    -fallbackOrderForNonUassetFiles (if order is not specified for ubulk/uexp files, figure out implicit order based on the uasset order. Generally applies only to the cooker order)"));
+
 	return false;
 }
diff --git a/Engine/Source/Programs/AutomationTool/Scripts/CopyBuildToStagingDirectory.Automation.cs b/Engine/Source/Programs/AutomationTool/Scripts/CopyBuildToStagingDirectory.Automation.cs
--- a/Engine/Source/Programs/AutomationTool/Scripts/CopyBuildToStagingDirectory.Automation.cs
+++ b/Engine/Source/Programs/AutomationTool/Scripts/CopyBuildToStagingDirectory.Automation.cs
@@ -1948,7 +1948,7 @@ public partial class Project : CommandUtils
                         }
 					}
 
-					string PatchSeekOptArgs = String.Empty;
+					string AdditionalArgs = String.Empty;
 					if (bShouldGeneratePatch)
 					{
 						string PatchSeekOptMode = String.Empty;
@@ -1957,23 +1957,30 @@ public partial class Project : CommandUtils
 						string PatchSeekOptMaxAdjacentOrderDiff = String.Empty;
 						if (PlatformGameConfig.GetString("/Script/UnrealEd.ProjectPackagingSettings", "PatchSeekOptMode", out PatchSeekOptMode))
 						{
-							PatchSeekOptArgs += String.Format(" -PatchSeekOptMode={0}", PatchSeekOptMode);
+							AdditionalArgs += String.Format(" -PatchSeekOptMode={0}", PatchSeekOptMode);
 						}
 						if (PlatformGameConfig.GetString("/Script/UnrealEd.ProjectPackagingSettings", "PatchSeekOptMaxGapSize", out PatchSeekOptMaxGapSize))
 						{
-							PatchSeekOptArgs += String.Format(" -PatchSeekOptMaxGapSize={0}", PatchSeekOptMaxGapSize);
+							AdditionalArgs += String.Format(" -PatchSeekOptMaxGapSize={0}", PatchSeekOptMaxGapSize);
 						}
 						if (PlatformGameConfig.GetString("/Script/UnrealEd.ProjectPackagingSettings", "PatchSeekOptMaxInflationPercent", out PatchSeekOptMaxInflationPercent))
 						{
-							PatchSeekOptArgs += String.Format(" -PatchSeekOptMaxInflationPercent={0}", PatchSeekOptMaxInflationPercent);
+							AdditionalArgs += String.Format(" -PatchSeekOptMaxInflationPercent={0}", PatchSeekOptMaxInflationPercent);
 						}
 						if (PlatformGameConfig.GetString("/Script/UnrealEd.ProjectPackagingSettings", "PatchSeekOptMaxAdjacentOrderDiff", out PatchSeekOptMaxAdjacentOrderDiff))
 						{
-							PatchSeekOptArgs += String.Format(" -PatchSeekOptMaxAdjacentOrderDiff={0}", PatchSeekOptMaxAdjacentOrderDiff);
+							AdditionalArgs += String.Format(" -PatchSeekOptMaxAdjacentOrderDiff={0}", PatchSeekOptMaxAdjacentOrderDiff);
 						}
 					}
 
-					Commands.Add(GetUnrealPakArguments(Params.RawProjectPath, PakParams.UnrealPakResponseFile, OutputLocation, PrimaryOrderFile, SC.StageTargetPlatform.GetPlatformPakCommandLine(Params, SC) + PatchSeekOptArgs + BulkOption + CompressionFormats + " " + Params.AdditionalPakOptions, PakParams.bCompressed, CryptoSettings, CryptoKeysCacheFilename, PatchSourceContentPath, PakParams.EncryptionKeyGuid, SecondaryOrderFile));
+					bool bPakFallbackOrderForNonUassetFiles = false;
+					PlatformGameConfig.GetBool("/Script/UnrealEd.ProjectPackagingSettings", "bPakFallbackOrderForNonUassetFiles", out bPakFallbackOrderForNonUassetFiles);
+					if (bPakFallbackOrderForNonUassetFiles)
+					{
+						AdditionalArgs += " -fallbackOrderForNonUassetFiles";
+					}
+
+					Commands.Add(GetUnrealPakArguments(Params.RawProjectPath, PakParams.UnrealPakResponseFile, OutputLocation, PrimaryOrderFile, SC.StageTargetPlatform.GetPlatformPakCommandLine(Params, SC) + AdditionalArgs + BulkOption + CompressionFormats + " " + Params.AdditionalPakOptions, PakParams.bCompressed, CryptoSettings, CryptoKeysCacheFilename, PatchSourceContentPath, PakParams.EncryptionKeyGuid, SecondaryOrderFile));
 					LogNames.Add(OutputLocation.GetFileNameWithoutExtension());
 				}
 			}
