Disable GUARANTEE_UASSET_AND_UEXP_IN_SAME_PAK for patch pak generation, since it's not necessary with delete records. Lifing this restriction makes pak seek optimization a lot more efficient.
Profiling in suggests it improves streaming performance, and the paks are smaller

diff --git a/Engine/Source/Developer/PakFileUtilities/Private/PakFileUtilities.cpp b/Engine/Source/Developer/PakFileUtilities/Private/PakFileUtilities.cpp
--- a/Engine/Source/Developer/PakFileUtilities/Private/PakFileUtilities.cpp
+++ b/Engine/Source/Developer/PakFileUtilities/Private/PakFileUtilities.cpp
@@ -23,7 +23,7 @@
 
 IMPLEMENT_MODULE(FDefaultModuleImpl, PakFileUtilities);
 
-#define GUARANTEE_UASSET_AND_UEXP_IN_SAME_PAK 1
+#define GUARANTEE_UASSET_AND_UEXP_IN_SAME_PAK 0
 
 #define SEEK_OPT_VERBOSITY Display
 
