Content build label update now errors if the build is not on prod, instead of trying to upload it again.
Added new content only build node  and now generates pak patch files for pak files not in original build.
Also merges manifest with original build
Also posts build to gamedev and prod
GenerateChunkDBs now takes in two additional parameters specifying how many rollback labels to process and which ones to optimize (preparation for future change where we process more chunkdbs).
Added support for generating buildinfo.ini for ps4.

diff --git a/Engine/Source/Developer/PakFileUtilities/Private/PakFileUtilities.cpp b/Engine/Source/Developer/PakFileUtilities/Private/PakFileUtilities.cpp
--- a/Engine/Source/Developer/PakFileUtilities/Private/PakFileUtilities.cpp
+++ b/Engine/Source/Developer/PakFileUtilities/Private/PakFileUtilities.cpp
@@ -376,6 +376,7 @@ struct FPakCommandLineParameters
 	bool   GeneratePatch;
 	FString SourcePatchPakFilename;
 	FString SourcePatchDiffDirectory;
+	FString InputFinalPakFilename; // This is the resulting pak file we want to end up with after we generate the pak patch.  This is used instead of passing in the raw content.
 	bool EncryptIndex;
 	bool UseCustomCompressor;
 	FGuid EncryptionKeyGuid;
@@ -429,6 +430,22 @@ struct FPakEntryOrder
 	uint64  Order;
 };
 
+
+struct FFileInfo
+{
+	uint64 FileSize;
+	int32 PatchIndex;
+	bool bIsDeleteRecord;
+	bool bForceInclude;
+	uint8 Hash[16];
+};
+
+bool ExtractFilesFromPak(const TCHAR* InPakFilename, TMap<FString, FFileInfo>& InFileHashes, 
+	const TCHAR* InDestPath, bool bUseMountPoint, 
+	const FKeyChain& InKeyChain, const FString* InFilter, 
+	TArray<FPakInputPair>* OutEntries = nullptr, TArray<FPakInputPair>* OutDeletedEntries = nullptr, 
+	FPakOrderMap* OutOrderMap = nullptr, TArray<FGuid>* OutUsedEncryptionKeys = nullptr, bool* OutAnyPakSigned = nullptr);
+
 struct FCompressedFileBuffer
 {
 	FCompressedFileBuffer()
@@ -472,6 +489,9 @@ struct FCompressedFileBuffer
 	TUniquePtr<uint8[]>		CompressedBuffer;
 };
 
+void LoadKeyChainFromFile(const FString& InFilename, FKeyChain& OutCryptoSettings);
+void ApplyEncryptionKeys(const FKeyChain& KeyChain);
+
 template <class T>
 bool ReadSizeParam(const TCHAR* CmdLine, const TCHAR* ParamStr, T& SizeOut)
 {
@@ -846,10 +866,16 @@ void ProcessCommandLine(const TCHAR* CmdLine, const TArray<FString>& NonOptionAr
 
 	if (FParse::Value(CmdLine, TEXT("-create="), ResponseFile))
 	{
-		TArray<FString> Lines;
+		
 
 		CmdLineParameters.GeneratePatch = FParse::Value(CmdLine, TEXT("-generatepatch="), CmdLineParameters.SourcePatchPakFilename);
 
+
+		bool bCompress = FParse::Param(CmdLine, TEXT("compress"));
+		bool bEncrypt = FParse::Param(CmdLine, TEXT("encrypt"));
+
+		
+
 		if (CmdLineParameters.GeneratePatch)
 		{
 			FParse::Value(CmdLine, TEXT("-patchSeekOptMaxInflationPercent="), CmdLineParameters.SeekOptParams.MaxInflationPercent);
@@ -862,93 +888,121 @@ void ProcessCommandLine(const TCHAR* CmdLine, const TArray<FString>& NonOptionAr
 				CmdLineParameters.SeekOptParams.Mode = ESeekOptMode::OnePass;
 			}
 			FParse::Value(CmdLine, TEXT("-patchSeekOptMode="), (int32&)CmdLineParameters.SeekOptParams.Mode);
-		}
-
-		bool bCompress = FParse::Param(CmdLine, TEXT("compress"));
-		bool bEncrypt = FParse::Param(CmdLine, TEXT("encrypt"));
 
-		bool bParseLines = true;
-		if (IFileManager::Get().DirectoryExists(*ResponseFile))
-		{
-			IFileManager::Get().FindFilesRecursive(Lines, *ResponseFile, TEXT("*"), true, false);
-			bParseLines = false;
+			
 		}
-		else
+
+		// if the response file is a pak file, then this is the pak file we want to use as the source
+		if (ResponseFile.EndsWith(TEXT(".pak"), ESearchCase::IgnoreCase) && CmdLineParameters.GeneratePatch)
 		{
-			FString Text;
-			UE_LOG(LogPakFile, Display, TEXT("Loading response file %s"), *ResponseFile);
-			if (FFileHelper::LoadFileToString(Text, *ResponseFile))
+			FString OutputPath;
+			if (FParse::Value(CmdLine, TEXT("extractedpaktemp="), OutputPath) == false)
 			{
-				// Remove all carriage return characters.
-				Text.ReplaceInline(TEXT("\r"), TEXT(""));
-				// Read all lines
-				Text.ParseIntoArray(Lines, TEXT("\n"), true);
+				UE_LOG(LogPakFile, Error, TEXT("-extractedpaktemp= not specified.  Required when specifying pak file as the response file."), *ResponseFile);
 			}
-			else
+
+			FString ExtractedPakKeysFile;
+			FKeyChain ExtractedPakKeys;
+			if ( FParse::Value(CmdLine, TEXT("extractedpakcryptokeys="), ExtractedPakKeysFile) )
 			{
-				UE_LOG(LogPakFile, Error, TEXT("Failed to load %s"), *ResponseFile);
+				LoadKeyChainFromFile(ExtractedPakKeysFile, ExtractedPakKeys);
+				ApplyEncryptionKeys(ExtractedPakKeys);
 			}
-		}
 
-		for (int32 EntryIndex = 0; EntryIndex < Lines.Num(); EntryIndex++)
+			TMap<FString, FFileInfo> FileHashes;
+			ExtractFilesFromPak(*ResponseFile, FileHashes, *OutputPath, true, ExtractedPakKeys, nullptr, &Entries, nullptr, nullptr, nullptr, nullptr);
+		}
+		else
 		{
-			TArray<FString> SourceAndDest;
-			TArray<FString> Switches;
-			if (bParseLines)
+			TArray<FString> Lines;
+			bool bParseLines = true;
+			if (IFileManager::Get().DirectoryExists(*ResponseFile))
 			{
-				Lines[EntryIndex].TrimStartInline();
-				CommandLineParseHelper(*Lines[EntryIndex], SourceAndDest, Switches);
+				IFileManager::Get().FindFilesRecursive(Lines, *ResponseFile, TEXT("*"), true, false);
+				bParseLines = false;
 			}
 			else
 			{
-				SourceAndDest.Add(Lines[EntryIndex]);
-			}
-			if( SourceAndDest.Num() == 0)
-			{
-				continue;
+				FString Text;
+				UE_LOG(LogPakFile, Display, TEXT("Loading response file %s"), *ResponseFile);
+				if (FFileHelper::LoadFileToString(Text, *ResponseFile))
+				{
+					// Remove all carriage return characters.
+					Text.ReplaceInline(TEXT("\r"), TEXT(""));
+					// Read all lines
+					Text.ParseIntoArray(Lines, TEXT("\n"), true);
+				}
+				else
+				{
+					UE_LOG(LogPakFile, Error, TEXT("Failed to load %s"), *ResponseFile);
+				}
 			}
-			FPakInputPair Input;
 
-			Input.Source = SourceAndDest[0];
-			FPaths::NormalizeFilename(Input.Source);
-			if (SourceAndDest.Num() > 1)
-			{
-				Input.Dest = FPaths::GetPath(SourceAndDest[1]);
-			}
-			else
+			for (int32 EntryIndex = 0; EntryIndex < Lines.Num(); EntryIndex++)
 			{
-				Input.Dest = FPaths::GetPath(Input.Source);
-			}
-			FPaths::NormalizeFilename(Input.Dest);
-			FPakFile::MakeDirectoryFromPath(Input.Dest);
+				TArray<FString> SourceAndDest;
+				TArray<FString> Switches;
+				if (bParseLines)
+				{
+					Lines[EntryIndex].TrimStartInline();
+					CommandLineParseHelper(*Lines[EntryIndex], SourceAndDest, Switches);
+				}
+				else
+				{
+					SourceAndDest.Add(Lines[EntryIndex]);
+				}
+				if (SourceAndDest.Num() == 0)
+				{
+					continue;
+				}
+				FPakInputPair Input;
 
-			//check for compression switches
-			for (int32 Index = 0; Index < Switches.Num(); ++Index)
-			{
-				if (Switches[Index] == TEXT("compress"))
+				Input.Source = SourceAndDest[0];
+				FPaths::NormalizeFilename(Input.Source);
+				if (SourceAndDest.Num() > 1)
 				{
-					Input.bNeedsCompression = true;
+					Input.Dest = FPaths::GetPath(SourceAndDest[1]);
 				}
-				if (Switches[Index] == TEXT("encrypt"))
+				else
 				{
-					Input.bNeedEncryption = true;
+					Input.Dest = FPaths::GetPath(Input.Source);
 				}
-			}
-			Input.bNeedsCompression |= bCompress;
-			Input.bNeedEncryption |= bEncrypt;
+				FPaths::NormalizeFilename(Input.Dest);
+				FPakFile::MakeDirectoryFromPath(Input.Dest);
 
-			UE_LOG(LogPakFile, Log, TEXT("Added file Source: %s Dest: %s"), *Input.Source, *Input.Dest);
+				//check for compression switches
+				for (int32 Index = 0; Index < Switches.Num(); ++Index)
+				{
+					if (Switches[Index] == TEXT("compress"))
+					{
+						Input.bNeedsCompression = true;
+					}
+					if (Switches[Index] == TEXT("encrypt"))
+					{
+						Input.bNeedEncryption = true;
+					}
+					if (Switches[Index] == TEXT("delete"))
+					{
+						Input.bIsDeleteRecord = true;
+						Input.Dest = Input.Source;
+					}
+				}
+				Input.bNeedsCompression |= bCompress;
+				Input.bNeedEncryption |= bEncrypt;
 
-			bool bIsMappedBulk = Input.Source.EndsWith(TEXT(".m.ubulk"));
-			if (bIsMappedBulk && CmdLineParameters.AlignForMemoryMapping > 0 && Input.bNeedsCompression 
-				&& !Input.bNeedEncryption) // if it is encrypted, we will compress it anyway since it won't be mapped at runtime
-			{
-				// no compression for bulk aligned files because they are memory mapped
-				Input.bNeedsCompression = false;
-				UE_LOG(LogPakFile, Log, TEXT("Stripped compression from %s for memory mapping."), *Input.Dest);
+				UE_LOG(LogPakFile, Log, TEXT("Added file Source: %s Dest: %s"), *Input.Source, *Input.Dest);
+
+				bool bIsMappedBulk = Input.Source.EndsWith(TEXT(".m.ubulk"));
+				if (bIsMappedBulk && CmdLineParameters.AlignForMemoryMapping > 0 && Input.bNeedsCompression
+					&& !Input.bNeedEncryption) // if it is encrypted, we will compress it anyway since it won't be mapped at runtime
+				{
+					// no compression for bulk aligned files because they are memory mapped
+					Input.bNeedsCompression = false;
+					UE_LOG(LogPakFile, Log, TEXT("Stripped compression from %s for memory mapping."), *Input.Dest);
+				}
+				Entries.Add(Input);
 			}
-			Entries.Add(Input);
-		}			
+		}
 	}
 	else
 	{
@@ -997,6 +1051,12 @@ void CollectFilesToAdd(TArray<FPakInputPair>& OutFilesToAdd, const TArray<FPakIn
 		bool bCompression = Input.bNeedsCompression;
 		bool bEncryption = Input.bNeedEncryption;
 
+		if (Input.bIsDeleteRecord)
+		{
+			// just pass through any delete records found in the input
+			OutFilesToAdd.Add(Input);
+			continue;
+		}
 
 		FString Filename = FPaths::GetCleanFilename(Source);
 		FString Directory = FPaths::GetPath(Source);
@@ -1100,7 +1160,7 @@ void CollectFilesToAdd(TArray<FPakInputPair>& OutFilesToAdd, const TArray<FPakIn
 	{
 		FORCEINLINE bool operator()(const FPakInputPair& A, const FPakInputPair& B) const
 		{
-			return A.SuggestedOrder == B.SuggestedOrder ? A.Dest < B.Dest : A.SuggestedOrder < B.SuggestedOrder;
+			return  A.bIsDeleteRecord == B.bIsDeleteRecord ? (A.SuggestedOrder == B.SuggestedOrder ? A.Dest < B.Dest : A.SuggestedOrder < B.SuggestedOrder) : A.bIsDeleteRecord < B.bIsDeleteRecord;
 		}
 	};
 	OutFilesToAdd.Sort(FInputPairSort());
@@ -1281,6 +1341,8 @@ void LoadKeyChainFromFile(const FString& InFilename, FKeyChain& OutCryptoSetting
 		}
 	}
 	delete File;
+	FGuid EncryptionKeyOverrideGuid;
+	OutCryptoSettings.MasterEncryptionKey = OutCryptoSettings.EncryptionKeys.Find(EncryptionKeyOverrideGuid);
 }
 
 void LoadKeyChain(const TCHAR* CmdLine, FKeyChain& OutCryptoSettings)
@@ -2602,16 +2664,10 @@ bool GeneratePIXMappingFile(const TArray<FString> InPakFileList, const FString&
 	return true;
 }
 
-struct FFileInfo
-{
-	uint64 FileSize;
-	int32 PatchIndex;
-	bool bIsDeleteRecord;
-	bool bForceInclude;
-	uint8 Hash[16];
-};
-
-bool ExtractFilesFromPak(const TCHAR* InPakFilename, TMap<FString, FFileInfo>& InFileHashes, const TCHAR* InDestPath, bool bUseMountPoint, const FKeyChain& InKeyChain, const FString* InFilter, TArray<FPakInputPair>* OutEntries = nullptr, TArray<FPakInputPair>* OutDeletedEntries = nullptr, FPakOrderMap* OutOrderMap = nullptr, TArray<FGuid>* OutUsedEncryptionKeys = nullptr, bool* OutAnyPakSigned = nullptr)
+bool ExtractFilesFromPak(const TCHAR* InPakFilename, TMap<FString, FFileInfo>& InFileHashes, 
+	const TCHAR* InDestPath, bool bUseMountPoint, const FKeyChain& InKeyChain, const FString* InFilter, 
+	TArray<FPakInputPair>* OutEntries, TArray<FPakInputPair>* OutDeletedEntries, FPakOrderMap* OutOrderMap, 
+	TArray<FGuid>* OutUsedEncryptionKeys, bool* OutAnyPakSigned)
 {
 	// Gather all patch versions of the requested pak file and run through each separately
 	TArray<FString> PakFileList;
@@ -2951,17 +3007,12 @@ bool GenerateHashesFromPak(const TCHAR* InPakFilename, const TCHAR* InDestPakFil
 {
 	OutLowestSourcePakVersion = FPakInfo::PakFile_Version_Invalid;
 
-	TArray<FString> FoundFiles;
-	IFileManager::Get().FindFiles(FoundFiles, InPakFilename, true, false);
-	if (FoundFiles.Num() == 0)
-	{
-		return false;
-	}
-
 	// Gather all patch pak files and run through them one at a time
 	TArray<FString> PakFileList;
+	IFileManager::Get().FindFiles(PakFileList, InPakFilename, true, false);
+
 	FString PakFileDirectory = FPaths::GetPath(InPakFilename);
-	IFileManager::Get().FindFiles(PakFileList, *PakFileDirectory, *FPaths::GetCleanFilename(InPakFilename));
+
 	for (int32 PakFileIndex = 0; PakFileIndex < PakFileList.Num(); PakFileIndex++)
 	{
 		FString PakFilename = PakFileDirectory + "\\" + PakFileList[PakFileIndex];
@@ -3919,12 +3970,15 @@ bool ExecuteUnrealPak(const TCHAR* CmdLine)
 	
 		if (NonOptionArguments.Num() != 2)
 		{
-			UE_LOG(LogPakFile, Error, TEXT("Incorrect arguments. Expected: -Extract <PakFile> <OutputPath>"));
+			UE_LOG(LogPakFile, Error, TEXT("Incorrect arguments. Expected: -Extract <PakFile> <OutputPath> [-responsefile=<filename>]"));
 			return false;
 		}
 
 		FString PakFilename = GetPakPath(*NonOptionArguments[0], false);
 
+		FString ResponseFileName;
+		bool bGenerateResponseFile = FParse::Value(CmdLine, TEXT("responseFile="), ResponseFileName);
+
 		bool bUseFilter = false;
 		FString Filter;
 		FString DestPath = NonOptionArguments[1];
@@ -3932,7 +3986,45 @@ bool ExecuteUnrealPak(const TCHAR* CmdLine)
 		bUseFilter = FParse::Value(CmdLine, TEXT("Filter="), Filter);
 		bool bExtractToMountPoint = FParse::Param(CmdLine, TEXT("ExtractToMountPoint"));
 		TMap<FString, FFileInfo> EmptyMap;
-		return ExtractFilesFromPak(*PakFilename, EmptyMap, *DestPath, bExtractToMountPoint, KeyChain, bUseFilter ? &Filter : nullptr);
+		TArray<FPakInputPair> ResponseContent;
+		TArray<FPakInputPair> DeletedContent;
+		if (ExtractFilesFromPak(*PakFilename, EmptyMap, *DestPath, bExtractToMountPoint, KeyChain, bUseFilter ? &Filter : nullptr, &ResponseContent, &DeletedContent) == false)
+		{
+			return false;
+		}
+
+		if (bGenerateResponseFile)
+		{
+			FArchive* ResponseArchive = IFileManager::Get().CreateFileWriter(*ResponseFileName);
+			ResponseArchive->SetIsTextFormat(true);
+			// generate a response file
+			if (FParse::Param(CmdLine, TEXT("includedeleted")))
+			{
+				for (int32 I = 0; I < DeletedContent.Num(); ++I)
+				{
+					ResponseArchive->Logf(TEXT("%s -delete"), *DeletedContent[I].Dest);
+				}
+			}
+
+			for (int32 I = 0; I < ResponseContent.Num(); ++I)
+			{
+				const FPakInputPair& Response = ResponseContent[I];
+				FString Line = FString::Printf(TEXT("%s %s"), *Response.Source, *Response.Dest);
+				if (Response.bNeedEncryption)
+				{
+					Line += " -encrypt";
+				}
+				if (Response.bNeedsCompression)
+				{
+					Line += " -compress";
+				}
+				ResponseArchive->Logf(TEXT("%s"), *Line);
+			}
+			ResponseArchive->Close();
+			delete ResponseArchive;
+		}
+
+		return true;
 	}
 
 	if (FParse::Param(CmdLine, TEXT("AuditFiles")))
diff --git a/Engine/Source/Programs/AutomationTool/AutomationUtils/AutomationUtils.Automation.csproj b/Engine/Source/Programs/AutomationTool/AutomationUtils/AutomationUtils.Automation.csproj
--- a/Engine/Source/Programs/AutomationTool/AutomationUtils/AutomationUtils.Automation.csproj
+++ b/Engine/Source/Programs/AutomationTool/AutomationUtils/AutomationUtils.Automation.csproj
@@ -58,6 +58,7 @@
     <Compile Include="AutomationException.cs" />
     <Compile Include="BuildCommand.cs" />
     <Compile Include="BuildUtils.cs" />
+    <Compile Include="BundleUtils.cs" />
     <Compile Include="CommandEnvironment.cs" />
     <Compile Include="DeviceReservation.cs" />
     <Compile Include="ImageUtils.cs" />
diff --git a/Engine/Source/Programs/AutomationTool/AutomationUtils/BundleUtils.cs b/Engine/Source/Programs/AutomationTool/AutomationUtils/BundleUtils.cs
new file mode 100644
--- /dev/null
+++ b/Engine/Source/Programs/AutomationTool/AutomationUtils/BundleUtils.cs
@@ -0,0 +1,121 @@
+﻿using System;
+using System.Collections.Generic;
+using System.Linq;
+using System.Text;
+using System.Threading.Tasks;
+using Tools.DotNETCommon;
+using AutomationTool;
+using UnrealBuildTool;
+
+namespace AutomationUtils.Automation
+{
+	public class BundleUtils
+	{
+		public interface IReadOnlyBundleSettings
+		{
+			string Name { get; }
+			string OverrideName { get; } // name used for ini section name in the BackgroundDownload.ini file
+			string ParentName { get;  }
+			List<string> Tags { get; }
+			List<BundleSettings> Children { get; }
+			bool bFoundParent { get; }
+			bool bContainsShaderLibrary { get; }
+		}
+
+
+		public class BundleSettings : IReadOnlyBundleSettings
+		{
+			public BundleSettings()
+			{
+				Children = new List<BundleSettings>();
+				bFoundParent = false;
+				bContainsShaderLibrary = false;
+			}
+
+			public string Name { get; set; }
+			public string OverrideName { get; set; }
+			public string ParentName { get; set; }
+			public List<string> Tags { get; set; }
+			public List<BundleSettings> Children { get; set; }
+			public List<string> FileRegex { get; set; }
+			public bool bFoundParent { get; set; }
+			public bool bContainsShaderLibrary { get; set; }
+			public int Priority { get; set; }
+		};
+		public static bool LoadBundleConfig(string BundleIniFile, ref Dictionary<string, BundleSettings> Bundles)
+		{
+			
+			if ( System.IO.File.Exists(BundleIniFile) == false )
+			{
+				CommandUtils.LogWarning("Unable to find bundle config ini file  {0}", BundleIniFile);
+				return false;
+			}
+
+			FileReference BundleFileReference = new FileReference(BundleIniFile);
+
+			ConfigHierarchy BundleConfig = new ConfigHierarchy(new ConfigFile[] { new ConfigFile(BundleFileReference) });
+			int PriorityCounter = 0;
+			foreach (string SectionName in BundleConfig.SectionNames)
+			{
+				BundleSettings Bundle = new BundleSettings();
+				Bundle.Name = SectionName;
+				Bundle.Priority = PriorityCounter;
+				++PriorityCounter;
+				{
+					string OverrideName;
+					if (BundleConfig.GetString(SectionName, "OverrideName", out OverrideName))
+					{
+						Bundle.OverrideName = OverrideName;
+					}
+					else
+					{
+						Bundle.OverrideName = Bundle.Name;
+					}
+				}
+				{
+					string ParentName;
+					BundleConfig.GetString(SectionName, "Parent", out ParentName);
+					Bundle.ParentName = ParentName;
+				}
+				{
+					List<string> Tags;
+					BundleConfig.GetArray(SectionName, "Tags", out Tags);
+					Bundle.Tags = Tags;
+				}
+				{
+					List<string> FileRegex;
+					BundleConfig.GetArray(SectionName, "FileRegex", out FileRegex);
+					Bundle.FileRegex = FileRegex;
+				}
+				{
+					bool bContainsShaderLibrary;
+					BundleConfig.GetBool(SectionName, "ContainsShaderLibrary", out bContainsShaderLibrary);
+					Bundle.bContainsShaderLibrary = bContainsShaderLibrary;
+				}
+				if (Bundle.Tags == null)
+				{
+					Bundle.Tags = new List<string>();
+				}
+
+				Bundles.Add(SectionName, Bundle);
+				//BundleConfig.GetArray(SectionName, "OptionalTags", out Bundle.OptionalTags);
+			}
+
+
+
+			foreach (var BundleIt in Bundles)
+			{
+				BundleSettings Bundle = BundleIt.Value;
+				if (Bundles.ContainsKey(Bundle.ParentName))
+				{
+					BundleSettings ParentBundle = Bundles[Bundle.ParentName];
+					ParentBundle.Children.Add(Bundle);
+					Bundle.bFoundParent = true;
+				}
+			}
+			return true;
+		}
+
+
+	}
+}
diff --git a/Engine/Source/Programs/AutomationTool/AutomationUtils/MCPPublic.cs b/Engine/Source/Programs/AutomationTool/AutomationUtils/MCPPublic.cs
--- a/Engine/Source/Programs/AutomationTool/AutomationUtils/MCPPublic.cs
+++ b/Engine/Source/Programs/AutomationTool/AutomationUtils/MCPPublic.cs
@@ -130,6 +130,11 @@ namespace EpicGames.MCP.Automation
 		/// AndroidCN Platform.
 		/// </summary>
 		AndroidCN,
+
+		/// <summary>
+		/// PS4 platform
+		/// </summary>
+		PS4,
 	}
 
     /// <summary>
diff --git a/Engine/Source/Programs/AutomationTool/Scripts/CopyBuildToStagingDirectory.Automation.cs b/Engine/Source/Programs/AutomationTool/Scripts/CopyBuildToStagingDirectory.Automation.cs
--- a/Engine/Source/Programs/AutomationTool/Scripts/CopyBuildToStagingDirectory.Automation.cs
+++ b/Engine/Source/Programs/AutomationTool/Scripts/CopyBuildToStagingDirectory.Automation.cs
@@ -165,7 +165,7 @@ public partial class Project : CommandUtils
 		return CmdLine.ToString();
 	}
 
-	static FileReference GetUnrealPakLocation()
+	public static FileReference GetUnrealPakLocation()
 	{
 		if (HostPlatform.Current.HostEditorPlatform == UnrealTargetPlatform.Win64)
 		{
