Adding internal -CalcCompressionBlockCRCs command line option to UnrealPak to try and track down the rare FPakAsyncReadFileHandle::DoProcessing issue + a bit more logging when it fails.

diff --git a/Engine/Source/Developer/PakFileUtilities/Private/PakFileUtilities.cpp b/Engine/Source/Developer/PakFileUtilities/Private/PakFileUtilities.cpp
--- a/Engine/Source/Developer/PakFileUtilities/Private/PakFileUtilities.cpp
+++ b/Engine/Source/Developer/PakFileUtilities/Private/PakFileUtilities.cpp
@@ -2569,7 +2569,7 @@ bool ListFilesAtOffset( const TCHAR* InPakFileName, const TArray<int64>& InOffse
 	FPakFile PakFile(&FPlatformFileManager::Get().GetPlatformFile(), InPakFileName, false);
 	if (!PakFile.IsValid())
 	{
-		UE_LOG(LogPakFile, Error, TEXT("Failed to open %s"), *InPakFileName );
+		UE_LOG(LogPakFile, Error, TEXT("Failed to open %s"), InPakFileName );
 		return false;
 	}
 
@@ -2609,6 +2609,116 @@ bool ListFilesAtOffset( const TCHAR* InPakFileName, const TArray<int64>& InOffse
 	return true;
 }
 
+// used for diagnosing errors in FPakAsyncReadFileHandle::RawReadCallback
+bool ShowCompressionBlockCRCs( const TCHAR* InPakFileName, TArray<int64>& InOffsets, const FKeyChain& InKeyChain )
+{
+	// open the pak file
+	FPakFile PakFile(&FPlatformFileManager::Get().GetPlatformFile(), InPakFileName, false);
+	if (!PakFile.IsValid())
+	{
+		UE_LOG(LogPakFile, Error, TEXT("Failed to open %s"), InPakFileName );
+		return false;
+	}
+
+	// read the pak file and iterate over all given offsets
+	FArchive& PakReader = *PakFile.GetSharedReader(NULL);
+	UE_LOG(LogPakFile, Display, TEXT("") );
+	for( int64 Offset : InOffsets )
+	{
+		//sanity check the offset
+		if (Offset < 0 || Offset > PakFile.TotalSize() )
+		{
+			UE_LOG(LogPakFile, Error, TEXT("Offset: %lld - out of range (max size is %lld)"), Offset, PakFile.TotalSize() );
+			continue;
+		}
+
+		//find the matching entry
+		const FPakEntry* Entry = nullptr;
+		for (FPakFile::FFileIterator It(PakFile); It; ++It)
+		{
+			const FPakEntry& ThisEntry = It.Info();
+			if( Offset >= ThisEntry.Offset && Offset <= ThisEntry.Offset+ThisEntry.Size )
+			{
+				Entry = &ThisEntry;
+				FString EntryFilename = It.Filename();
+				FName EntryCompressionMethod = PakFile.GetInfo().GetCompressionMethod(Entry->CompressionMethodIndex);
+
+				UE_LOG(LogPakFile, Display, TEXT("Offset: %lld  -> EntrySize: %lld  Encrypted: %-3s  Compression: %-8s  [%s]"), Offset, Entry->Size, Entry->IsEncrypted() ? TEXT("Yes") : TEXT("No"), *EntryCompressionMethod.ToString(), *EntryFilename );
+				break;
+			}
+		}
+		if (Entry == nullptr)
+		{
+			UE_LOG(LogPakFile, Error, TEXT("Offset: %lld - no entry found."), Offset );
+			continue;
+		}
+
+		// sanity check
+		if (Entry->CompressionMethodIndex == 0)
+		{
+			UE_LOG(LogPakFile, Error, TEXT("    Entry isn't compressed (not supported)") );
+			continue;
+		}
+		if (Entry->IsDeleteRecord())
+		{
+			UE_LOG(LogPakFile, Error, TEXT("    Entry is deleted") );
+			continue;
+		}
+
+		//iterate over all blocks, decoding them and computing the checksum
+
+		//... adapted from UncompressCopyFile...
+		FName EntryCompressionMethod = PakFile.GetInfo().GetCompressionMethod(Entry->CompressionMethodIndex);
+		int32 MaxCompressionBlockSize = FCompression::CompressMemoryBound(EntryCompressionMethod, Entry->CompressionBlockSize);
+		for (const FPakCompressedBlock& Block : Entry->CompressionBlocks)
+		{
+			MaxCompressionBlockSize = FMath::Max<int32>(MaxCompressionBlockSize, Block.CompressedEnd - Block.CompressedStart);
+		}
+
+		int64 WorkingSize = Entry->CompressionBlockSize + MaxCompressionBlockSize;
+		uint8* PersistentBuffer = (uint8*)FMemory::Malloc(WorkingSize);
+
+		for (uint32 BlockIndex=0, BlockIndexNum=Entry->CompressionBlocks.Num(); BlockIndex < BlockIndexNum; ++BlockIndex)
+		{
+			uint32 CompressedBlockSize = Entry->CompressionBlocks[BlockIndex].CompressedEnd - Entry->CompressionBlocks[BlockIndex].CompressedStart;
+			uint32 UncompressedBlockSize = (uint32)FMath::Min<int64>(Entry->UncompressedSize - Entry->CompressionBlockSize*BlockIndex, Entry->CompressionBlockSize);
+			PakReader.Seek(Entry->CompressionBlocks[BlockIndex].CompressedStart + (PakFile.GetInfo().HasRelativeCompressedChunkOffsets() ? Entry->Offset : 0));
+			uint32 SizeToRead = Entry->IsEncrypted() ? Align(CompressedBlockSize, FAES::AESBlockSize) : CompressedBlockSize;
+			PakReader.Serialize(PersistentBuffer, SizeToRead);
+
+			if (Entry->IsEncrypted())
+			{
+				const FNamedAESKey* Key = InKeyChain.MasterEncryptionKey;
+				check(Key);
+				FAES::DecryptData(PersistentBuffer, SizeToRead, Key->Key);
+			}
+
+			// adapted from FPakAsyncReadFileHandle::RawReadCallback
+			int32 ProcessedSize = Entry->CompressionBlockSize;
+			if (BlockIndex == BlockIndexNum - 1)
+			{
+				ProcessedSize = Entry->UncompressedSize % Entry->CompressionBlockSize;
+				if (!ProcessedSize)
+				{
+					ProcessedSize = Entry->CompressionBlockSize; // last block was a full block
+				}
+			}
+
+			// compute checksum and log out the block information
+			const uint32 BlockCrc32 = FCrc::MemCrc32( PersistentBuffer, CompressedBlockSize );
+			const FString HexBytes = BytesToHex(PersistentBuffer,FMath::Min(CompressedBlockSize,32U));
+			UE_LOG(LogPakFile, Display, TEXT("    Block:%-6d  ProcessedSize: %-6d  DecompressionRawSize: %-6d  Crc32: %-12u [%s...]"), BlockIndex, ProcessedSize, CompressedBlockSize, BlockCrc32, *HexBytes );
+		}
+
+		FMemory::Free(PersistentBuffer);
+		UE_LOG(LogPakFile, Display, TEXT("") );
+
+	}
+
+	UE_LOG(LogPakFile, Display, TEXT("done") );
+	return true;
+}
+
 bool GeneratePIXMappingFile(const TArray<FString> InPakFileList, const FString& OutputPath)
 {
 	if (!InPakFileList.Num())
@@ -4080,6 +4190,28 @@ bool ExecuteUnrealPak(const TCHAR* CmdLine)
 
 		return ListFilesAtOffset( *PakFilename, Offsets );
 	}
+
+	if (FParse::Param(CmdLine, TEXT("CalcCompressionBlockCRCs")))
+	{
+		if (NonOptionArguments.Num() < 2)
+		{
+			UE_LOG(LogPakFile, Error, TEXT("Incorrect arguments. Expected -CalcCompressionBlockCRCs <PakFile> [Offset...] ") );
+			return false;
+		}
+
+		FString PakFilename = GetPakPath(*NonOptionArguments[0], false);
+
+		TArray<int64> Offsets;
+		for( int ArgI = 1; ArgI < NonOptionArguments.Num(); ArgI++ )
+		{
+			if( FCString::IsNumeric(*NonOptionArguments[ArgI]) )
+			{
+				Offsets.Add( FCString::Strtoi64( *NonOptionArguments[ArgI], nullptr, 10 ) );
+			}
+		}
+
+		return ShowCompressionBlockCRCs( *PakFilename, Offsets, KeyChain );
+	}
 	
 	if (FParse::Param(CmdLine, TEXT("GeneratePIXMappingFile")))
 	{
diff --git a/Engine/Source/Runtime/PakFile/Private/IPlatformFilePak.cpp b/Engine/Source/Runtime/PakFile/Private/IPlatformFilePak.cpp
--- a/Engine/Source/Runtime/PakFile/Private/IPlatformFilePak.cpp
+++ b/Engine/Source/Runtime/PakFile/Private/IPlatformFilePak.cpp
@@ -3506,7 +3506,8 @@ public:
 
 			if( !FCompression::UncompressMemory(CompressionMethod, Output, Block.ProcessedSize, Block.Raw, Block.DecompressionRawSize) )
 			{
-				UE_LOG( LogPakFile, Fatal, TEXT("Pak Decompression failed. PakFile: %s. EntryOffset: %lld, EntrySize: %lld, CompressionMethod:%s Output:%p  ProcessedSize:%d  Buf:%p  Block.DecompressionRawSize:%d  Crc32:%u"), *PakFile.ToString(), FileEntry.Offset, FileEntry.Size, *CompressionMethod.ToString(), Output, Block.ProcessedSize, Block.Raw, Block.DecompressionRawSize, FCrc::MemCrc32( Block.Raw, Block.DecompressionRawSize ) );
+				const FString HexBytes = BytesToHex(Block.Raw,FMath::Min(Block.DecompressionRawSize,32));
+				UE_LOG( LogPakFile, Fatal, TEXT("Pak Decompression failed. PakFile:%s, EntryOffset:%lld, EntrySize:%lld, Method:%s, ProcessedSize:%d, RawSize:%d, Crc32:%u, BlockIndex:%d, Encrypt:%d, Delete:%d, Output:%p, Raw:%p, Processed:%p, Bytes:[%s...]"), *PakFile.ToString(), FileEntry.Offset, FileEntry.Size, *CompressionMethod.ToString(), Block.ProcessedSize, Block.DecompressionRawSize, FCrc::MemCrc32( Block.Raw, Block.DecompressionRawSize ), Block.BlockIndex, FileEntry.IsEncrypted()?1:0, FileEntry.IsDeleteRecord()?1:0, Output, Block.Raw, Block.Processed, *HexBytes );
 			}
 			FMemory::Free(Block.Raw);
 			Block.Raw = nullptr;
