UnrealPak - Added "TestMemoryOptimization" mode which attempts to unload filenames and shrink pak entries across every pak in a directory, making sure there are no collisons and that everything works correctly

diff --git a/Engine/Source/Developer/PakFileUtilities/Private/PakFileUtilities.cpp b/Engine/Source/Developer/PakFileUtilities/Private/PakFileUtilities.cpp
--- a/Engine/Source/Developer/PakFileUtilities/Private/PakFileUtilities.cpp
+++ b/Engine/Source/Developer/PakFileUtilities/Private/PakFileUtilities.cpp
@@ -4026,6 +4026,54 @@ bool ExecuteUnrealPak(const TCHAR* CmdLine)
 		return TestPakFile(*PakFilename);
 	}
 
+	if (FParse::Param(CmdLine, TEXT("TestMemoryOptimization")))
+	{
+		TArray<FPakInputPair> Entries;
+		FPakCommandLineParameters CmdLineParameters;
+		ProcessCommandLine(CmdLine, NonOptionArguments, Entries, CmdLineParameters);
+
+		if (NonOptionArguments.Num() != 1)
+		{
+			UE_LOG(LogPakFile, Error, TEXT("Incorrect arguments. Expected: -TestMemoryOptimization <SourceFolder>"));
+			return false;
+		}
+
+		FString SourceDir = NonOptionArguments[0];
+		TArray<FString> PakFilenames;
+		IFileManager::Get().FindFiles(PakFilenames, *SourceDir, TEXT("*.pak"));
+		TArray<FPakFile*> PakFiles;
+		PakFiles.Empty(PakFilenames.Num());
+
+		for (const FString& PakFilename : PakFilenames)
+		{
+			PakFiles.Add(new FPakFile(&FPlatformFileManager::Get().GetPlatformFile(), *(FPaths::Combine(SourceDir, PakFilename)), false));
+		}
+
+		TMap<uint64, FPakEntry> CollisionChecker;
+		
+		for (FPakFile* PakFile : PakFiles)
+		{
+			if (!PakFile->UnloadPakEntryFilenames(CollisionChecker, nullptr, false))
+			{
+				UE_LOG(LogPakFile, Error, TEXT("Pak '%s' failed to unload filenames"), *PakFile->GetFilename());
+			}
+
+			if (!PakFile->ShrinkPakEntriesMemoryUsage())
+			{
+				UE_LOG(LogPakFile, Error, TEXT("Pak '%s' failed to shrink entries"), *PakFile->GetFilename());
+			}
+		}
+
+		for (FPakFile* PakFile : PakFiles)
+		{
+			delete PakFile;
+		}
+
+		PakFiles.Empty();
+
+		return true;
+	}
+
 	if (FParse::Param(CmdLine, TEXT("List")))
 	{
 		TArray<FPakInputPair> Entries;
diff --git a/Engine/Source/Runtime/PakFile/Private/IPlatformFilePak.cpp b/Engine/Source/Runtime/PakFile/Private/IPlatformFilePak.cpp
--- a/Engine/Source/Runtime/PakFile/Private/IPlatformFilePak.cpp
+++ b/Engine/Source/Runtime/PakFile/Private/IPlatformFilePak.cpp
@@ -4677,7 +4677,7 @@ static inline int32 CDECL CompareFMiniFileEntry(const void* Left, const void* Ri
 	return 0;
 }
 
-bool FPakFile::UnloadPakEntryFilenames(TMap<uint64, FPakEntry>& CrossPakCollisionChecker, TArray<FString>* DirectoryRootsToKeep)
+bool FPakFile::UnloadPakEntryFilenames(TMap<uint64, FPakEntry>& CrossPakCollisionChecker, TArray<FString>* DirectoryRootsToKeep, bool bAllowRetries)
 {
 	// If the process has already been done, get out of here.
 	if (bAttemptedPakFilenameUnload || bFilenamesRemoved)
@@ -4694,7 +4694,7 @@ bool FPakFile::UnloadPakEntryFilenames(TMap<uint64, FPakEntry>& CrossPakCollisio
 
 	// Variables for the filename hashing and collision detection.
 	int NumRetries = 0;
-	const int MAX_RETRIES = 10;
+	const int MAX_RETRIES = bAllowRetries ? 10 : 1;
 	bool bHasCollision;
 	FilenameStartHash = FCrc::StrCrc32(*GetFilename());
     
@@ -5117,6 +5117,10 @@ bool FPakFile::ShrinkPakEntriesMemoryUsage()
 	// space of the original anymore.
 	Files.Empty(0);
 
+	static uint64 Total = 0;
+	Total += TotalSizeOfCompressedEntries;
+	UE_LOG(LogPakFile, Display, TEXT("Compressed pak entries down to %d bytes [Total = %llu bytes]"), TotalSizeOfCompressedEntries, Total);
+
 	return true;
 }
 
diff --git a/Engine/Source/Runtime/PakFile/Public/IPlatformFilePak.h b/Engine/Source/Runtime/PakFile/Public/IPlatformFilePak.h
--- a/Engine/Source/Runtime/PakFile/Public/IPlatformFilePak.h
+++ b/Engine/Source/Runtime/PakFile/Public/IPlatformFilePak.h
@@ -957,8 +957,9 @@ public:
 	 *
 	 * @param CrossPakCollisionChecker A map of hash->fileentry records encountered during filename unloading on other pak files. Used to detect collisions with entries in other pak files.
 	 * @param DirectoryRootsToKeep An array of strings in wildcard format that specify whole directory structures of filenames to keep in memory for directory iteration to work.
+	 * @param bAllowRetries If a collision is encountered, change the intial seed and try again a fixed number of times before failing
 	 */
-	bool UnloadPakEntryFilenames(TMap<uint64, FPakEntry>& CrossPakCollisionChecker, TArray<FString>* DirectoryRootsToKeep = nullptr);
+	bool UnloadPakEntryFilenames(TMap<uint64, FPakEntry>& CrossPakCollisionChecker, TArray<FString>* DirectoryRootsToKeep = nullptr, bool bAllowRetries = true);
 
 	/**
 	 * Lower memory usage by bit-encoding the pak file entry information.
