Unrealpak: Throw a meaningful error when a crypto keys cache is specified on the command line that doesn't exist

diff --git a/Engine/Source/Developer/PakFileUtilities/Private/PakFileUtilities.cpp b/Engine/Source/Developer/PakFileUtilities/Private/PakFileUtilities.cpp
--- a/Engine/Source/Developer/PakFileUtilities/Private/PakFileUtilities.cpp
+++ b/Engine/Source/Developer/PakFileUtilities/Private/PakFileUtilities.cpp
@@ -954,6 +954,7 @@ TEncryptionInt ParseEncryptionIntFromJson(TSharedPtr<FJsonObject> InObj, const T
 void PrepareEncryptionAndSigningKeysFromCryptoKeyCache(const FString& InFilename, FKeyPair& OutSigningKey, TKeyChain& OutKeyChain)
 {
 	FArchive* File = IFileManager::Get().CreateFileReader(*InFilename);
+	UE_CLOG(File == nullptr, LogPakFile, Fatal, TEXT("Specified crypto keys cache '%s' does not exist!"), *InFilename);
 	TSharedPtr<FJsonObject> RootObject;
 	TSharedRef<TJsonReader<char>> Reader = TJsonReaderFactory<char>::Create(File);
 	if (FJsonSerializer::Deserialize(Reader, RootObject))
