Removed signing key generation functionality from unrealpak. Keys should be generated with the CryptoKeys plugin in the editor now

diff --git a/Engine/Source/Developer/PakFileUtilities/Private/KeyGenerator.cpp b/Engine/Source/Developer/PakFileUtilities/Private/KeyGenerator.cpp
deleted file mode 100644
diff --git a/Engine/Source/Developer/PakFileUtilities/Private/KeyGenerator.h b/Engine/Source/Developer/PakFileUtilities/Private/KeyGenerator.h
deleted file mode 100644
diff --git a/Engine/Source/Developer/PakFileUtilities/Private/PakFileUtilities.cpp b/Engine/Source/Developer/PakFileUtilities/Private/PakFileUtilities.cpp
--- a/Engine/Source/Developer/PakFileUtilities/Private/PakFileUtilities.cpp
+++ b/Engine/Source/Developer/PakFileUtilities/Private/PakFileUtilities.cpp
@@ -5,7 +5,6 @@
 #include "Misc/SecureHash.h"
 #include "Math/BigInt.h"
 #include "SignedArchiveWriter.h"
-#include "KeyGenerator.h"
 #include "Misc/AES.h"
 #include "Templates/UniquePtr.h"
 #include "Serialization/LargeMemoryWriter.h"
@@ -20,6 +19,71 @@
 #include "HAL/PlatformFilemanager.h"
 #include "Async/ParallelFor.h"
 
+/**
+ * Encryption keys: public and private
+ */
+struct FKeyPair
+{
+	/** Public decryption key */
+	FEncryptionKey PublicKey;
+	/** Private encryption key */
+	FEncryptionKey PrivateKey;
+
+	friend FArchive& operator<<(FArchive& Ar, FKeyPair& Pair)
+	{
+		Ar << Pair.PublicKey.Exponent;
+		Ar << Pair.PublicKey.Modulus;
+		Ar << Pair.PrivateKey.Exponent;
+		Ar << Pair.PrivateKey.Modulus;
+		return Ar;
+	}
+
+	bool IsValid() const
+	{
+		return !PrivateKey.Exponent.IsZero()
+			&& !PrivateKey.Modulus.IsZero()
+			&& !PublicKey.Exponent.IsZero()
+			&& !PublicKey.Modulus.IsZero();
+	}
+};
+
+bool TestKeys(FKeyPair& Pair)
+{
+	UE_LOG(LogPakFile, Display, TEXT("Testing signature keys."));
+
+	// Just some random values
+	static TEncryptionInt TestData[] =
+	{
+		11,
+		253,
+		128,
+		234,
+		56,
+		89,
+		34,
+		179,
+		29,
+		1024,
+		(int64)(MAX_int32),
+		(int64)(MAX_uint32)-1
+	};
+
+	for (int32 TestIndex = 0; TestIndex < ARRAY_COUNT(TestData); ++TestIndex)
+	{
+		TEncryptionInt EncryptedData = FEncryption::ModularPow(TestData[TestIndex], Pair.PrivateKey.Exponent, Pair.PrivateKey.Modulus);
+		TEncryptionInt DecryptedData = FEncryption::ModularPow(EncryptedData, Pair.PublicKey.Exponent, Pair.PublicKey.Modulus);
+		if (TestData[TestIndex] != DecryptedData)
+		{
+			UE_LOG(LogPakFile, Error, TEXT("Keys do not properly encrypt/decrypt data (failed test with %lld)"), TestData[TestIndex].ToInt());
+			return false;
+		}
+	}
+
+	UE_LOG(LogPakFile, Display, TEXT("Signature keys check completed successfuly."));
+
+	return true;
+}
+
 struct FNamedAESKey
 {
 	FString Name;
@@ -1199,34 +1263,6 @@ void PrepareEncryptionAndSigningKeys(const TCHAR* CmdLine, FKeyPair& OutSigningK
 			OutKeyChain.Add(NewKey.Guid, NewKey);
 			UE_LOG(LogPakFile, Display, TEXT("Parsed AES encryption key from command line."));
 		}
-
-		FString KeyFilename;
-		if (FParse::Value(CmdLine, TEXT("sign="), KeyFilename, false))
-		{
-			if (KeyFilename.StartsWith(TEXT("0x")))
-			{
-				TArray<FString> KeyValueText;
-				int32 NumParts = KeyFilename.ParseIntoArray(KeyValueText, TEXT("+"), true);
-				if (NumParts == 3)
-				{
-					OutSigningKey.PrivateKey.Exponent.Parse(KeyValueText[0]);
-					OutSigningKey.PrivateKey.Modulus.Parse(KeyValueText[1]);
-					OutSigningKey.PublicKey.Exponent.Parse(KeyValueText[2]);
-					OutSigningKey.PublicKey.Modulus = OutSigningKey.PrivateKey.Modulus;
-
-					UE_LOG(LogPakFile, Display, TEXT("Parsed signature keys from command line."));
-				}
-				else
-				{
-					UE_LOG(LogPakFile, Error, TEXT("Expected 3, got %d, when parsing %s"), KeyValueText.Num(), *KeyFilename);
-					OutSigningKey.PrivateKey.Exponent.Zero();
-				}
-			}
-			else if (!ReadKeysFromFile(*KeyFilename, OutSigningKey))
-			{
-				UE_LOG(LogPakFile, Error, TEXT("Unable to load signature keys %s."), *KeyFilename);
-			}
-		}
 	}
 
 	if (OutSigningKey.IsValid())
@@ -3341,10 +3377,6 @@ bool Repack(const FString& InputPakFile, const FString& OutputPakFile, const FPa
  *   -Sign=filename use the key pair in filename to sign a pak file, or: -sign=key_hex_values_separated_with_+, i.e: -sign=0x123456789abcdef+0x1234567+0x12345abc
  *    where the first number is the private key exponend, the second one is modulus and the third one is the public key exponent.
  *   -Signed use with -extract and -test to let the code know this is a signed pak
- *   -GenerateKeys=filename generates encryption key pair for signing a pak file
- *   -P=prime will use a predefined prime number for generating encryption key file
- *   -Q=prime same as above, P != Q, GCD(P, Q) = 1 (which is always true if they're both prime)
- *   -GeneratePrimeTable=filename generates a prime table for faster prime number generation (.inl file)
  *   -TableMax=number maximum prime number in the generated table (default is 10000)
  *
  * @param	ArgC	Command-line argument count
@@ -3389,27 +3421,6 @@ bool ExecuteUnrealPak(const TCHAR* CmdLine)
 		return Result;
 	}
 
-	FString KeyFilename;
-	if (FParse::Value(CmdLine, TEXT("GenerateKeys="), KeyFilename, false))
-	{
-		return GenerateKeys(*KeyFilename);
-	}
-
-	if (FParse::Value(CmdLine, TEXT("GeneratePrimeTable="), KeyFilename, false))
-	{
-		int64 MaxPrimeValue = 10000;
-		FParse::Value(CmdLine, TEXT("TableMax="), MaxPrimeValue);
-		GeneratePrimeNumberTable(MaxPrimeValue, *KeyFilename);
-		return true;
-	}
-
-	if (FParse::Param(CmdLine, TEXT("TestEncryption")))
-	{
-		void TestEncryption();
-		TestEncryption();
-		return true;
-	}
-
 	if (FParse::Param(CmdLine, TEXT("Test")))
 	{
 		if (NonOptionArguments.Num() != 1)
@@ -3773,8 +3784,6 @@ bool ExecuteUnrealPak(const TCHAR* CmdLine)
 	UE_LOG(LogPakFile, Error, TEXT("  UnrealPak <PakFilename> -Create=<ResponseFile> [Options]"));
 	UE_LOG(LogPakFile, Error, TEXT("  UnrealPak <PakFilename> -Dest=<MountPoint>"));
 	UE_LOG(LogPakFile, Error, TEXT("  UnrealPak <PakFilename> -Repack [-Output=Path] [-ExcludeDeleted] [Options]"));
-	UE_LOG(LogPakFile, Error, TEXT("  UnrealPak GenerateKeys=<KeyFilename>"));
-	UE_LOG(LogPakFile, Error, TEXT("  UnrealPak GeneratePrimeTable=<KeyFilename> [-TableMax=<N>]"));
 	UE_LOG(LogPakFile, Error, TEXT("  UnrealPak <PakFilename1> <PakFilename2> -diff"));
 	UE_LOG(LogPakFile, Error, TEXT("  UnrealPak <PakFolder> -AuditFiles [-OnlyDeleted] [-CSV=<filename>] [-order=<OrderingFile>] [-SortByOrdering]"));
 	UE_LOG(LogPakFile, Error, TEXT("  UnrealPak <PakFilename> -WhatsAtOffset [offset1] [offset2] [offset3] [...]"));
diff --git a/Engine/Source/Developer/PakFileUtilities/Private/Primes.inl b/Engine/Source/Developer/PakFileUtilities/Private/Primes.inl
deleted file mode 100644
diff --git a/Engine/Source/Developer/PakFileUtilities/Private/SignedArchiveWriter.cpp b/Engine/Source/Developer/PakFileUtilities/Private/SignedArchiveWriter.cpp
--- a/Engine/Source/Developer/PakFileUtilities/Private/SignedArchiveWriter.cpp
+++ b/Engine/Source/Developer/PakFileUtilities/Private/SignedArchiveWriter.cpp
@@ -99,64 +99,4 @@ int64 FSignedArchiveWriter::TotalSize()
 void FSignedArchiveWriter::Seek(int64 InPos)
 {
 	UE_LOG(LogPakFile, Fatal, TEXT("Seek is not supported in FSignedArchiveWriter."));
-}
-
-/**
- * Useful code for testing the encryption methods
- */
-void TestEncryption()
-{
-	FEncryptionKey PublicKey;
-	FEncryptionKey PrivateKey;
-	TEncryptionInt P(TEXT("0x21443BD2DD63E995403"));
-	TEncryptionInt Q(TEXT("0x28CBB6E5749AC65749"));
-	FEncryption::GenerateKeyPair(P, Q, PublicKey, PrivateKey);
-
-	// Generate random data
-	const int32 DataSize = 1024;
-	uint8* Data = (uint8*)FMemory::Malloc(DataSize);
-	for (int32 Index = 0; Index < DataSize; ++Index)
-	{
-		Data[Index] = (uint8)(Index % 255);
-	}
-
-	// Generate signature
-	FDecryptedSignature OriginalSignature;
-	FEncryptedSignature EncryptedSignature;
-	FDecryptedSignature DecryptedSignature;
-	OriginalSignature.Data = FCrc::MemCrc32(Data, DataSize);
-
-	// Encrypt signature with the private key
-	FEncryption::EncryptSignature(OriginalSignature, EncryptedSignature, PrivateKey);
-	FEncryption::DecryptSignature(EncryptedSignature, DecryptedSignature, PublicKey);
-
-	// Check
-	if (OriginalSignature == DecryptedSignature)
-	{
-		UE_LOG(LogPakFile, Display, TEXT("Keys match"));
-	}
-	else
-	{
-		UE_LOG(LogPakFile, Fatal, TEXT("Keys mismatched!"));
-	}
-
-	double OverallTime = 0.0;
-	double OverallNumTests = 0.0;
-	for (int32 TestIndex = 0; TestIndex < 10; ++TestIndex)
-	{
-		static const int64 NumTests = 500;
-		double Timer = FPlatformTime::Seconds();
-		{
-			for (int64 i = 0; i < NumTests; ++i)
-			{
-				FEncryption::DecryptSignature(EncryptedSignature, DecryptedSignature, PublicKey);
-			}
-		}
-		Timer = FPlatformTime::Seconds() - Timer;
-		OverallTime += Timer;
-		OverallNumTests += (double)NumTests;
-		UE_LOG(LogPakFile, Display, TEXT("%i signatures decrypted in %.4fs, Avg = %.4fs, OverallAvg = %.4fs"), NumTests, Timer, Timer / (double)NumTests, OverallTime / OverallNumTests);
-	}
-	
-	FMemory::Free(Data);
-}
+}
\ No newline at end of file
