Move bulk of the code for UnrealPak into an engine developer module, so it can be used in the editor.

diff --git a/Engine/Source/Developer/PakFileUtilities/PakFileUtilities.Build.cs b/Engine/Source/Developer/PakFileUtilities/PakFileUtilities.Build.cs
new file mode 100644
--- /dev/null
+++ b/Engine/Source/Developer/PakFileUtilities/PakFileUtilities.Build.cs
@@ -0,0 +1,22 @@
+// Copyright 1998-2018 Epic Games, Inc. All Rights Reserved.
+
+using UnrealBuildTool;
+
+public class PakFileUtilities : ModuleRules
+{
+	public PakFileUtilities(ReadOnlyTargetRules Target) : base(Target)
+	{
+		PrivateDependencyModuleNames.AddRange(new string[] { "Core", "CoreUObject", "PakFile", "Json", "Projects", "ApplicationCore" });
+
+        PrivateIncludePathModuleNames.AddRange(
+            new string[] {
+                "AssetRegistry",
+                "Json"
+        });
+
+        DynamicallyLoadedModuleNames.AddRange(
+            new string[] {
+                "AssetRegistry"
+        });
+    }
+}
diff --git a/Engine/Source/Developer/PakFileUtilities/Private/KeyGenerator.cpp b/Engine/Source/Developer/PakFileUtilities/Private/KeyGenerator.cpp
new file mode 100644
--- /dev/null
+++ b/Engine/Source/Developer/PakFileUtilities/Private/KeyGenerator.cpp
@@ -0,0 +1,532 @@
+// Copyright 1998-2018 Epic Games, Inc. All Rights Reserved.
+
+#include "KeyGenerator.h"
+#include "IPlatformFilePak.h"
+#include "Misc/SecureHash.h"
+#include "Math/BigInt.h"
+#include "Async/TaskGraphInterfaces.h"
+#include "HAL/Runnable.h"
+#include "HAL/RunnableThread.h"
+#include "Math/RandomStream.h"
+#include "Misc/FileHelper.h"
+#include "Misc/CommandLine.h"
+#include "Primes.inl"
+
+// Global constants
+namespace
+{
+	const TEncryptionInt Two(2);
+	const TEncryptionInt IterationStep(1000);
+	TArray<TEncryptionInt> PrimeLookupTable;
+}
+
+/**
+ * A thread that finds factors in the given range
+ */
+class FPrimeCheckRunnable : public FRunnable
+{
+	/** Flag indicating if a factor has been found. Shared across multiple threads. */
+	FThreadSafeCounter&	FoundFactor;
+	/** Candidate for a prime number */
+	TEncryptionInt PotentialPrime;
+	/** Start of a range to check for factors */
+	TEncryptionInt InitialValue;
+	/** End of a range to check for factors */
+	TEncryptionInt MaxValue;
+	/** This thread */
+	FRunnableThread* Thread;
+
+public:
+
+	FPrimeCheckRunnable(FThreadSafeCounter& InFoundFactor, TEncryptionInt Candidate, TEncryptionInt InInitialValue, TEncryptionInt InMaxValue)
+		: FoundFactor(InFoundFactor)
+		, PotentialPrime(Candidate)
+		, InitialValue(InInitialValue)
+		, MaxValue(InMaxValue)
+	{
+		// Must be odd number
+		check(!(Candidate & TEncryptionInt::One).IsZero());
+		Thread = FRunnableThread::Create(this, TEXT("FPrimeCheckRunnable"));
+	}
+
+	virtual ~FPrimeCheckRunnable()
+	{
+		delete Thread;
+		Thread = NULL;
+	}
+
+	FRunnableThread* GetThread()
+	{
+		return Thread;
+	}
+
+	// Start FRunnable interface
+	virtual bool Init() override { return true; }
+	virtual uint32 Run() override
+	{
+		TEncryptionInt Remainder;
+		int32 FactorCheckTimer = 0;
+		for (TEncryptionInt Factor = InitialValue; InitialValue <= MaxValue; Factor += Two)
+		{
+			TEncryptionInt Dividend(PotentialPrime);
+			Dividend.DivideWithRemainder(Factor, Remainder);
+			if (Remainder.IsZero())
+			{
+				FoundFactor.Increment();
+				break;
+			}
+			// Don't check the FoundFactor too often
+			FactorCheckTimer++;
+			if (FactorCheckTimer >= 100)
+			{
+				FactorCheckTimer = 0;
+				if (FoundFactor.GetValue() != 0)
+				{
+					// Another thread found a factor
+					break;
+				}
+			}
+		}
+		return 0;
+	}
+	// End FRunnable interface
+
+};
+
+/**
+ * Checks if the value is a prime number.
+ */
+bool IsPrime(const TEncryptionInt& InValue, bool bUseTasks)
+{
+	// 2 is but we don't care about small numbers here.
+	if ((InValue & TEncryptionInt::One) == 0)
+	{
+		return false;
+	}
+
+	TEncryptionInt Remainder;
+
+	// Check against known prime factors
+	int32 Index;
+	for (Index = 0; Index < PrimeLookupTable.Num() && PrimeLookupTable[Index] < InValue; ++Index)
+	{
+		TEncryptionInt Dividend(InValue);
+		Dividend.DivideWithRemainder(PrimeLookupTable[Index], Remainder);
+		if (Remainder.IsZero())
+		{
+			return false;
+		}
+	}
+	// This means the number is smaller than one of the primes in the prime table and it has no factors
+	if (Index < PrimeLookupTable.Num())
+	{
+		return true;
+	}
+
+	// Brute force, check all odd numbers > MaxKnownPrime < sqrt(Number)
+	TEncryptionInt MaxFactorValue(InValue);
+	MaxFactorValue.Sqrt();
+	TEncryptionInt Factor(PrimeLookupTable[PrimeLookupTable.Num() - 1] + Two); 
+	if (bUseTasks)
+	{
+		// Mutithreaded path. Split the range we look for factors over multiple threads. If one thread finds a factor
+		// we stop and reject this number.
+		// The worst case is when we actually have a prime number.
+		UE_LOG(LogPakFile, Display, TEXT("Detected potentially prime number %s. This may take a while..."), *InValue.ToString());
+
+		const int32 TaskCount = FPlatformMisc::NumberOfCoresIncludingHyperthreads();
+		TArray<FPrimeCheckRunnable*> Tasks;
+		Tasks.Reserve(TaskCount);
+		FThreadSafeCounter FoundFactors(0);
+
+		TEncryptionInt Range(MaxFactorValue - Factor);
+		Range /= TaskCount;
+
+		// Spawn threads
+		for (int32 TaskIndex = 0; TaskIndex < TaskCount; ++TaskIndex)
+		{
+			TEncryptionInt MaxValue(Factor + Range);
+			Tasks.Add(new FPrimeCheckRunnable(FoundFactors, InValue, Factor, MaxValue)); 
+			Factor = MaxValue;
+			if ((Factor & TEncryptionInt::One) == 0)
+			{
+				++Factor;
+			}
+		}
+		// Wait for all threads to complete
+		for (int32 TaskIndex = 0; TaskIndex < Tasks.Num(); ++TaskIndex)
+		{
+			Tasks[TaskIndex]->GetThread()->WaitForCompletion();
+			delete Tasks[TaskIndex];
+		}
+
+		if (FoundFactors.GetValue() > 0)
+		{
+			UE_LOG(LogPakFile, Display, TEXT("%s is not prime."), *InValue.ToString());
+			return false;
+		}
+		else
+		{
+			UE_LOG(LogPakFile, Display, TEXT("%s is prime!"), *InValue.ToString());
+		}
+	}
+	else
+	{		
+		// Single threaded path (used for generating prime table)
+		while (Factor < MaxFactorValue)
+		{
+			TEncryptionInt Dividend(InValue);
+			Dividend.DivideWithRemainder(Factor, Remainder);
+			if (Remainder.IsZero())
+			{
+				return false;
+			}
+			Factor += Two;
+		}
+	}
+	return true;
+}
+
+/**
+ * Generate two random prime numbers
+ */
+void GeneratePrimeNumbers(TEncryptionInt& P, TEncryptionInt& Q)
+{
+	// Generate a random odd number
+	FRandomStream Rand((int32)(FDateTime::Now().GetTicks() % (int64)MAX_int32));
+	uint32 RandBits[256/32] = 
+	{ 
+		0xffffffff, 
+		(uint32)Rand.RandRange(0, MAX_int32 - 1), 
+		(uint32)Rand.RandRange(0, MAX_int32 - 1), 
+		0, //(uint32)Rand.RandRange(0, MAX_int32 - 1) | 0xa0ff0000, 
+		0, 0, 0, 0 
+	};
+	TEncryptionInt InitialValue(RandBits);
+
+	// We need two primes
+	TArray<TEncryptionInt> DiscoveredPrimes;
+
+	int64 IterationCounter = 0;
+	double TimeAccumulator = 0.0;
+	const double StartTime = FPlatformTime::Seconds();
+	do
+	{		
+		if (IterationCounter == 10)
+		{
+			IterationCounter = 0;
+			TimeAccumulator = 0.0;
+		}
+		IterationCounter++;
+
+		const double IterationStartTime = FPlatformTime::Seconds();
+
+		TEncryptionInt MinValue(InitialValue - IterationStep);
+		while (InitialValue >= MinValue && DiscoveredPrimes.Num() < 2)
+		{
+			if (IsPrime(InitialValue, true))
+			{
+				DiscoveredPrimes.Add(InitialValue);
+			}
+			InitialValue -= Two;
+		}
+	}
+	while (DiscoveredPrimes.Num() < 2);
+
+	UE_LOG(LogPakFile, Display, TEXT("Generated prime numbers in %.2lfs."), FPlatformTime::Seconds() - StartTime);
+	P = DiscoveredPrimes[0];
+	Q = DiscoveredPrimes[1];
+	UE_LOG(LogPakFile, Display, TEXT("P=%s"), *P.ToString());
+	UE_LOG(LogPakFile, Display, TEXT("Q=%s"), *Q.ToString());
+}
+
+/**
+ * Lookup table generation - fill it with precompile primes
+ */
+void FillPrimeLookupTableWithPrecompiledNumbers()
+{
+	const int32 PrimeTableLength = ARRAY_COUNT(PrimeTable);
+	PrimeLookupTable.Reserve(PrimeTableLength * PrimeTableLength);
+	for (int32 Index = 0; Index < PrimeTableLength; ++Index)
+	{
+		PrimeLookupTable.Add(PrimeTable[Index]);
+	}
+}
+
+void GeneratePrimeNumberTable(int64 MaxValue, const TCHAR* Filename)
+{
+	FillPrimeLookupTableWithPrecompiledNumbers();
+
+	UE_LOG(LogPakFile, Display, TEXT("Generating prime number table <= %lld: %s."), MaxValue, Filename);
+
+	FString PrimeTableString(TEXT("// Copyright 1998-2018 Epic Games, Inc. All Rights Reserved.\nTEncryptionInt PrimeTable[] = \n{\n\t2, "));
+	int64 PrimeCount = 1;
+	const double StartTime = FPlatformTime::Seconds();
+	for (int64 SmallNumber = 3; SmallNumber <= MaxValue; SmallNumber += 2)
+	{
+		if (IsPrime(SmallNumber, false))
+		{
+			PrimeTableString += FString::Printf(TEXT("%lld, "), SmallNumber);
+			PrimeCount++;
+			if ((PrimeCount % 10) == 0)
+			{
+				PrimeTableString += TEXT("\n\t");
+			}
+		}
+	}
+	PrimeTableString += TEXT("\n};\n");
+	UE_LOG(LogPakFile, Display, TEXT("Generated %lld primes in %.4lfs."), PrimeCount, FPlatformTime::Seconds() - StartTime);
+	FFileHelper::SaveStringToFile(PrimeTableString, Filename);
+}
+
+
+/**
+ * Multithreaded prime number generation (for the prime table)
+ * Finds prime numbers in the given range.
+ */
+class FPrimeFinderRunnable : public FRunnable
+{
+	/** Range start */
+	TEncryptionInt MinValue;
+	/** Range end */
+	TEncryptionInt MaxValue;
+	/** This thread */
+	FRunnableThread* Thread;
+	/** All primes found in the given range */
+	TArray<TEncryptionInt> FoundPrimes;
+
+public:
+
+	FPrimeFinderRunnable(TEncryptionInt InMinValue, TEncryptionInt InMaxValue)
+		: MinValue(InMinValue)
+		, MaxValue(InMaxValue)
+	{
+		// Must be an odd number
+		check(!(MinValue & TEncryptionInt::One).IsZero());
+		Thread = FRunnableThread::Create(this, TEXT("FPrimeFinderRunnable"));
+	}
+
+	virtual ~FPrimeFinderRunnable()
+	{
+		delete Thread;
+		Thread = NULL;
+	}
+
+	FRunnableThread* GetThread()
+	{
+		return Thread;
+	}
+
+	const TArray<TEncryptionInt>& GetFoundPrimes() const
+	{
+		return FoundPrimes;
+	}
+
+	// Start FRunnable interface
+	virtual bool Init() override { return true; }
+	virtual uint32 Run() override
+	{
+		TEncryptionInt Remainder;
+		int32 FactorCheckTimer = 0;
+		for (TEncryptionInt Candidate = MinValue; Candidate <= MaxValue; Candidate += Two)
+		{
+			if (IsPrime(Candidate, false))
+			{
+				FoundPrimes.Add(Candidate);
+			}
+		}
+		return 0;
+	}
+	// End FRunnable interface
+
+};
+
+/**
+ * Generates a lookup table in runtime.
+ * This is a superset of precompiled prime table and primes generated on startup.
+ * The reason for this is that the precompiled table can't be too big because it affects
+ * the compile times and we don't usually use UnrealPak for prime number generation anyway.
+ */
+void GeneratePrimeNumberLookupTable()
+{	
+	const int32 PrimeTableLength = ARRAY_COUNT(PrimeTable);
+	UE_LOG(LogPakFile, Display, TEXT("Generating prime number lookup table (max size: %d)."), PrimeTableLength * PrimeTableLength);
+	const double StartTime = FPlatformTime::Seconds();
+
+	FillPrimeLookupTableWithPrecompiledNumbers();
+
+	TEncryptionInt MinPrimeValue(PrimeLookupTable[PrimeLookupTable.Num() - 1] + Two);
+	TEncryptionInt MaxPrimeValue(MinPrimeValue);
+	MaxPrimeValue *= 100;
+
+	const int32 TaskCount = FPlatformMisc::NumberOfCoresIncludingHyperthreads();
+	TArray<FPrimeFinderRunnable*> Tasks;
+	Tasks.Reserve(TaskCount);
+
+	TEncryptionInt Range(MaxPrimeValue - MinPrimeValue);
+	Range /= TaskCount;
+
+	for (int32 TaskIndex = 0; TaskIndex < TaskCount; ++TaskIndex)
+	{
+		TEncryptionInt MaxValue(MinPrimeValue + Range);
+		Tasks.Add(new FPrimeFinderRunnable(MinPrimeValue, MaxValue)); 
+		MinPrimeValue = MaxValue;
+		if ((MinPrimeValue & TEncryptionInt::One) == 0)
+		{
+			++MinPrimeValue;
+		}
+	}
+
+	TArray<TEncryptionInt> NewPrimes;
+	for (int32 TaskIndex = 0; TaskIndex < Tasks.Num(); ++TaskIndex)
+	{
+		Tasks[TaskIndex]->GetThread()->WaitForCompletion();
+		NewPrimes.Append(Tasks[TaskIndex]->GetFoundPrimes());
+		delete Tasks[TaskIndex];
+	}
+	PrimeLookupTable.Append(NewPrimes);
+
+	UE_LOG(LogPakFile, Display, TEXT("Generated %d primes in %.4lfs."), PrimeLookupTable.Num(), FPlatformTime::Seconds() - StartTime);
+}
+
+bool GenerateKeys(const TCHAR* KeyFilename)
+{
+	UE_LOG(LogPakFile, Display, TEXT("Generating keys %s."), KeyFilename);
+
+	TEncryptionInt P;
+	TEncryptionInt Q;
+
+	FString CmdLineP;
+	FString CmdLineQ;
+	FParse::Value(FCommandLine::Get(), TEXT("P="), CmdLineP);
+	FParse::Value(FCommandLine::Get(), TEXT("Q="), CmdLineQ);
+
+	P.Parse(CmdLineP);
+	Q.Parse(CmdLineQ);
+
+	const bool bNoVerifyPrimes = FParse::Param(FCommandLine::Get(), TEXT("NoVerifyPrimes"));
+
+	// Check if we have valid primes in the command line.
+	// @todo: IsPrime check should probably go when we start to use big primes
+	bool bGeneratePrimes = !(P > Two && Q > Two);
+
+	if (bGeneratePrimes || !bNoVerifyPrimes)
+	{
+		GeneratePrimeNumberLookupTable();
+	}
+
+	if (!bGeneratePrimes && !bNoVerifyPrimes)
+	{
+		if (!IsPrime(P, false))
+		{
+			UE_LOG(LogPakFile, Warning, TEXT("P=%s is not prime!"), *CmdLineP);
+			bGeneratePrimes = true;
+		}
+		if (!IsPrime(Q, false))
+		{
+			UE_LOG(LogPakFile, Warning, TEXT("Q=%s is not prime!"), *CmdLineQ);
+			bGeneratePrimes = true;
+		}
+	}
+
+	if (bGeneratePrimes)
+	{
+		// Generate random prime numbers
+		UE_LOG(LogPakFile, Display, TEXT("Generating prime numbers..."));
+		GeneratePrimeNumbers(P, Q);
+	}
+	else
+	{
+		// Use predefined primes
+		UE_LOG(LogPakFile, Display, TEXT("Using predefined values to generate keys."));
+	}
+
+	// Generate key pair
+	UE_LOG(LogPakFile, Display, TEXT("Generating key pair..."));
+	FKeyPair Keys;
+	FEncryption::GenerateKeyPair(P, Q, Keys.PublicKey, Keys.PrivateKey);
+
+	if (TestKeys(Keys))
+	{
+		return SaveKeysToFile(Keys, KeyFilename);
+	}
+	else
+	{
+		return false;
+	}
+}
+
+bool SaveKeysToFile(const FKeyPair& Keys, const TCHAR* KeyFilename)
+{
+	UE_LOG(LogPakFile, Display, TEXT("Saving key pair in %s"), KeyFilename);	
+	FString KeyFileContents = FString::Printf(TEXT("%s\n%s\n%s"), *Keys.PrivateKey.Exponent.ToString(), *Keys.PrivateKey.Modulus.ToString(), *Keys.PublicKey.Exponent.ToString());
+	return FFileHelper::SaveStringToFile(KeyFileContents, KeyFilename);
+}
+
+bool ReadKeysFromFile(const TCHAR* KeyFilename, FKeyPair& OutKeys)
+{
+	bool bResult = false;
+	UE_LOG(LogPakFile, Display, TEXT("Loading key pair from %s"), KeyFilename);
+	FString KeyFileContents;
+	if (FFileHelper::LoadFileToString(KeyFileContents, KeyFilename))
+	{
+		TArray<FString> KeyValues;
+		KeyFileContents.ParseIntoArrayWS(KeyValues);
+		if (KeyValues.Num() != 3)
+		{
+			UE_LOG(LogPakFile, Error, TEXT("Expecting 3 values in %s, got %d."), KeyFilename, KeyValues.Num());
+		}
+		else
+		{
+			OutKeys.PrivateKey.Exponent.Parse(KeyValues[0]);
+			OutKeys.PrivateKey.Modulus.Parse(KeyValues[1]);
+			OutKeys.PublicKey.Exponent.Parse(KeyValues[2]);
+			OutKeys.PublicKey.Modulus = OutKeys.PrivateKey.Modulus;
+			bResult = true;
+		}
+	}
+	else
+	{
+		UE_LOG(LogPakFile, Error, TEXT("Failed to load key pair from %s"), KeyFilename);
+	}
+
+	return bResult;
+}
+
+
+bool TestKeys(FKeyPair& Pair)
+{
+	UE_LOG(LogPakFile, Display, TEXT("Testing signature keys."));
+
+	// Just some random values
+	static TEncryptionInt TestData[] =
+	{
+		11,
+		253,
+		128,
+		234,
+		56,
+		89,
+		34,
+		179,
+		29,
+		1024,
+		(int64)(MAX_int32),
+		(int64)(MAX_uint32) - 1
+	};
+
+	for (int32 TestIndex = 0; TestIndex < ARRAY_COUNT(TestData); ++TestIndex)
+	{
+		TEncryptionInt EncryptedData = FEncryption::ModularPow(TestData[TestIndex], Pair.PrivateKey.Exponent, Pair.PrivateKey.Modulus);
+		TEncryptionInt DecryptedData = FEncryption::ModularPow(EncryptedData, Pair.PublicKey.Exponent, Pair.PublicKey.Modulus);
+		if (TestData[TestIndex] != DecryptedData)
+		{
+			UE_LOG(LogPakFile, Error, TEXT("Keys do not properly encrypt/decrypt data (failed test with %lld)"), TestData[TestIndex].ToInt());
+			return false;
+		}
+	}
+
+	UE_LOG(LogPakFile, Display, TEXT("Signature keys check completed successfuly."));
+
+	return true;
+}
diff --git a/Engine/Source/Developer/PakFileUtilities/Private/KeyGenerator.h b/Engine/Source/Developer/PakFileUtilities/Private/KeyGenerator.h
new file mode 100644
--- /dev/null
+++ b/Engine/Source/Developer/PakFileUtilities/Private/KeyGenerator.h
@@ -0,0 +1,48 @@
+// Copyright 1998-2018 Epic Games, Inc. All Rights Reserved.
+#pragma once
+
+#include "CoreMinimal.h"
+#include "Math/BigInt.h"
+
+/**
+ * Encryption keys: public and private
+ */
+struct FKeyPair
+{
+	/** Public decryption key */
+	FEncryptionKey PublicKey;
+	/** Private encryption key */
+	FEncryptionKey PrivateKey;
+
+	friend FArchive& operator<<(FArchive& Ar, FKeyPair& Pair)
+	{
+		Ar << Pair.PublicKey.Exponent;
+		Ar << Pair.PublicKey.Modulus;
+		Ar << Pair.PrivateKey.Exponent;
+		Ar << Pair.PrivateKey.Modulus;
+		return Ar;
+	}
+
+	bool IsValid() const
+	{
+		return !PrivateKey.Exponent.IsZero()
+			&& !PrivateKey.Modulus.IsZero()
+			&& !PublicKey.Exponent.IsZero()
+			&& !PublicKey.Modulus.IsZero();
+	}
+};
+
+/**
+ * Generates a prime number table where the max prime number value is less than MaxValue
+ */
+void GeneratePrimeNumberTable(int64 MaxValue, const TCHAR* Filename);
+/**
+ * Generates a file with the encryption keys
+ */
+bool GenerateKeys(const TCHAR* KeyFilename);
+
+bool SaveKeysToFile(const FKeyPair& Keys, const TCHAR* KeyFilename);
+
+bool ReadKeysFromFile(const TCHAR* KeyFilename, FKeyPair& OutKeys);
+
+bool TestKeys(FKeyPair& Pair);
diff --git a/Engine/Source/Developer/PakFileUtilities/Private/PakFileUtilities.cpp b/Engine/Source/Developer/PakFileUtilities/Private/PakFileUtilities.cpp
new file mode 100644
--- /dev/null
+++ b/Engine/Source/Developer/PakFileUtilities/Private/PakFileUtilities.cpp
@@ -0,0 +1,2745 @@
+// Copyright 1998-2018 Epic Games, Inc. All Rights Reserved.
+
+#include "PakFileUtilities.h"
+#include "IPlatformFilePak.h"
+#include "Misc/SecureHash.h"
+#include "Math/BigInt.h"
+#include "SignedArchiveWriter.h"
+#include "KeyGenerator.h"
+#include "Misc/AES.h"
+#include "Templates/UniquePtr.h"
+#include "Serialization/LargeMemoryWriter.h"
+#include "AssetRegistryModule.h"
+#include "ProfilingDebugging/DiagnosticTable.h"
+#include "Serialization/JsonSerializer.h"
+#include "Misc/Base64.h"
+#include "Misc/Compression.h"
+#include "Features/IModularFeatures.h"
+#include "Misc/CoreDelegates.h"
+
+struct FPakCommandLineParameters
+{
+	FPakCommandLineParameters()
+		: CompressionBlockSize(64*1024)
+		, CompressionBitWindow(DEFAULT_ZLIB_BIT_WINDOW)
+		, FileSystemBlockSize(0)
+		, PatchFilePadAlign(0)
+		, GeneratePatch(false)
+		, EncryptIndex(false)
+		, UseCustomCompressor(false)
+		, OverridePlatformCompressor(false)
+	{}
+
+	int32  CompressionBlockSize;
+	int32  CompressionBitWindow;
+	int64  FileSystemBlockSize;
+	int64  PatchFilePadAlign;
+	bool   GeneratePatch;
+	FString SourcePatchPakFilename;
+	FString SourcePatchDiffDirectory;
+	bool EncryptIndex;
+	bool UseCustomCompressor;
+	bool OverridePlatformCompressor;
+};
+
+struct FPakEntryPair
+{
+	FString Filename;
+	FPakEntry Info;
+};
+
+struct FPakInputPair
+{
+	FString Source;
+	FString Dest;
+	uint64 SuggestedOrder; 
+	bool bNeedsCompression;
+	bool bNeedEncryption;
+
+	FPakInputPair()
+		: SuggestedOrder(MAX_uint64)
+		, bNeedsCompression(false)
+		, bNeedEncryption(false)
+	{}
+
+	FPakInputPair(const FString& InSource, const FString& InDest)
+		: Source(InSource)
+		, Dest(InDest)
+		, bNeedsCompression(false)
+		, bNeedEncryption(false)
+	{}
+
+	FORCEINLINE bool operator==(const FPakInputPair& Other) const
+	{
+		return Source == Other.Source;
+	}
+};
+
+struct FPakEntryOrder
+{
+	FPakEntryOrder() : Order(MAX_uint64) {}
+	FString Filename;
+	uint64  Order;
+};
+
+struct FCompressedFileBuffer
+{
+	FCompressedFileBuffer()
+		: OriginalSize(0)
+		,TotalCompressedSize(0)
+		,FileCompressionBlockSize(0)
+		,CompressedBufferSize(0)
+	{
+
+	}
+
+	void Reinitialize(FArchive* File,ECompressionFlags CompressionMethod,int64 CompressionBlockSize)
+	{
+		OriginalSize = File->TotalSize();
+		TotalCompressedSize = 0;
+		FileCompressionBlockSize = 0;
+		FileCompressionMethod = CompressionMethod;
+		CompressedBlocks.Reset();
+		CompressedBlocks.AddUninitialized((OriginalSize+CompressionBlockSize-1)/CompressionBlockSize);
+	}
+
+	void EnsureBufferSpace(int64 RequiredSpace)
+	{
+		if(RequiredSpace > CompressedBufferSize)
+		{
+			uint8* NewPtr = (uint8*)FMemory::Malloc(RequiredSpace);
+			FMemory::Memcpy(NewPtr,CompressedBuffer.Get(),CompressedBufferSize);
+			CompressedBuffer.Reset(NewPtr);
+			CompressedBufferSize = RequiredSpace;
+		}
+	}
+
+	bool CompressFileToWorkingBuffer(const FPakInputPair& InFile,uint8*& InOutPersistentBuffer,int64& InOutBufferSize,ECompressionFlags CompressionMethod,const int32 CompressionBlockSize,const int32 CompressionBitWindow);
+
+	int64						 OriginalSize;
+	int64						 TotalCompressedSize;
+	int32						 FileCompressionBlockSize;
+	ECompressionFlags			 FileCompressionMethod;
+	TArray<FPakCompressedBlock>  CompressedBlocks;
+	int64						 CompressedBufferSize;
+	TUniquePtr<uint8>		     CompressedBuffer;
+};
+
+FString GetLongestPath(TArray<FPakInputPair>& FilesToAdd)
+{
+	FString LongestPath;
+	int32 MaxNumDirectories = 0;
+
+	for (int32 FileIndex = 0; FileIndex < FilesToAdd.Num(); FileIndex++)
+	{
+		FString& Filename = FilesToAdd[FileIndex].Dest;
+		int32 NumDirectories = 0;
+		for (int32 Index = 0; Index < Filename.Len(); Index++)
+		{
+			if (Filename[Index] == '/')
+			{
+				NumDirectories++;
+			}
+		}
+		if (NumDirectories > MaxNumDirectories)
+		{
+			LongestPath = Filename;
+			MaxNumDirectories = NumDirectories;
+		}
+	}
+	return FPaths::GetPath(LongestPath) + TEXT("/");
+}
+
+FString GetCommonRootPath(TArray<FPakInputPair>& FilesToAdd)
+{
+	FString Root = GetLongestPath(FilesToAdd);
+	for (int32 FileIndex = 0; FileIndex < FilesToAdd.Num() && Root.Len(); FileIndex++)
+	{
+		FString Filename(FilesToAdd[FileIndex].Dest);
+		FString Path = FPaths::GetPath(Filename) + TEXT("/");
+		int32 CommonSeparatorIndex = -1;
+		int32 SeparatorIndex = Path.Find(TEXT("/"), ESearchCase::CaseSensitive);
+		while (SeparatorIndex >= 0)
+		{
+			if (FCString::Strnicmp(*Root, *Path, SeparatorIndex + 1) != 0)
+			{
+				break;
+			}
+			CommonSeparatorIndex = SeparatorIndex;
+			if (CommonSeparatorIndex + 1 < Path.Len())
+			{
+				SeparatorIndex = Path.Find(TEXT("/"), ESearchCase::CaseSensitive, ESearchDir::FromStart, CommonSeparatorIndex + 1);
+			}
+			else
+			{
+				break;
+			}
+		}
+		if ((CommonSeparatorIndex + 1) < Root.Len())
+		{
+			Root = Root.Mid(0, CommonSeparatorIndex + 1);
+		}
+	}
+	return Root;
+}
+
+bool FCompressedFileBuffer::CompressFileToWorkingBuffer(const FPakInputPair& InFile, uint8*& InOutPersistentBuffer, int64& InOutBufferSize, ECompressionFlags CompressionMethod, const int32 CompressionBlockSize, const int32 CompressionBitWindow)
+{
+	TUniquePtr<FArchive> FileHandle(IFileManager::Get().CreateFileReader(*InFile.Source));
+	if(!FileHandle)
+	{
+		TotalCompressedSize = 0;
+		return false;
+	}
+
+	Reinitialize(FileHandle.Get(), CompressionMethod, CompressionBlockSize);
+	const int64 FileSize = OriginalSize;
+	const int64 PaddedEncryptedFileSize = Align(FileSize,FAES::AESBlockSize);
+	if(InOutBufferSize < PaddedEncryptedFileSize)
+	{
+		InOutPersistentBuffer = (uint8*)FMemory::Realloc(InOutPersistentBuffer,PaddedEncryptedFileSize);
+		InOutBufferSize = FileSize;
+	}
+
+	// Load to buffer
+	FileHandle->Serialize(InOutPersistentBuffer,FileSize);
+
+	// Build buffers for working
+	int64 UncompressedSize = FileSize;
+	int32 CompressionBufferSize = Align(FCompression::CompressMemoryBound(CompressionMethod,CompressionBlockSize,CompressionBitWindow),FAES::AESBlockSize);
+	EnsureBufferSpace(Align(FCompression::CompressMemoryBound(CompressionMethod,FileSize,CompressionBitWindow),FAES::AESBlockSize));
+
+
+	TotalCompressedSize = 0;
+	int64 UncompressedBytes = 0;
+	int32 CurrentBlock = 0;
+	while(UncompressedSize)
+	{
+		int32 BlockSize = (int32)FMath::Min<int64>(UncompressedSize,CompressionBlockSize);
+		int32 MaxCompressedBlockSize = FCompression::CompressMemoryBound(CompressionMethod, BlockSize, CompressionBitWindow);
+		int32 CompressedBlockSize = FMath::Max<int32>(CompressionBufferSize, MaxCompressedBlockSize);
+		FileCompressionBlockSize = FMath::Max<uint32>(BlockSize, FileCompressionBlockSize);
+		EnsureBufferSpace(Align(TotalCompressedSize+CompressedBlockSize,FAES::AESBlockSize));
+		if(!FCompression::CompressMemory(CompressionMethod,CompressedBuffer.Get()+TotalCompressedSize,CompressedBlockSize,InOutPersistentBuffer+UncompressedBytes,BlockSize,CompressionBitWindow))
+		{
+			return false;
+		}
+		UncompressedSize -= BlockSize;
+		UncompressedBytes += BlockSize;
+
+		CompressedBlocks[CurrentBlock].CompressedStart = TotalCompressedSize;
+		CompressedBlocks[CurrentBlock].CompressedEnd = TotalCompressedSize+CompressedBlockSize;
+		++CurrentBlock;
+
+		TotalCompressedSize += CompressedBlockSize;
+
+		if(InFile.bNeedEncryption)
+		{
+			int32 EncryptionBlockPadding = Align(TotalCompressedSize,FAES::AESBlockSize);
+			for(int64 FillIndex=TotalCompressedSize; FillIndex < EncryptionBlockPadding; ++FillIndex)
+			{
+				// Fill the trailing buffer with bytes from file. Note that this is now from a fixed location
+				// rather than a random one so that we produce deterministic results
+				CompressedBuffer.Get()[FillIndex] = CompressedBuffer.Get()[FillIndex % TotalCompressedSize];
+			}
+			TotalCompressedSize += EncryptionBlockPadding - TotalCompressedSize;
+		}
+	}
+
+	return true;
+}
+
+bool PrepareCopyFileToPak(const FString& InMountPoint, const FPakInputPair& InFile, uint8*& InOutPersistentBuffer, int64& InOutBufferSize, FPakEntryPair& OutNewEntry, uint8*& OutDataToWrite, int64& OutSizeToWrite, const FAES::FAESKey& InEncryptionKey)
+{	
+	TUniquePtr<FArchive> FileHandle(IFileManager::Get().CreateFileReader(*InFile.Source));
+	bool bFileExists = FileHandle.IsValid();
+	if (bFileExists)
+	{
+		const int64 FileSize = FileHandle->TotalSize();
+		const int64 PaddedEncryptedFileSize = Align(FileSize, FAES::AESBlockSize); 
+		OutNewEntry.Filename = InFile.Dest.Mid(InMountPoint.Len());
+		OutNewEntry.Info.Offset = 0; // Don't serialize offsets here.
+		OutNewEntry.Info.Size = FileSize;
+		OutNewEntry.Info.UncompressedSize = FileSize;
+		OutNewEntry.Info.CompressionMethod = COMPRESS_None;
+		OutNewEntry.Info.bEncrypted = InFile.bNeedEncryption && InEncryptionKey.IsValid();
+
+		if (InOutBufferSize < PaddedEncryptedFileSize)
+		{
+			InOutPersistentBuffer = (uint8*)FMemory::Realloc(InOutPersistentBuffer, PaddedEncryptedFileSize);
+			InOutBufferSize = FileSize;
+		}
+
+		// Load to buffer
+		FileHandle->Serialize(InOutPersistentBuffer, FileSize);
+
+		{
+			OutSizeToWrite = FileSize;
+			if (InFile.bNeedEncryption && InEncryptionKey.IsValid())
+			{
+				for(int64 FillIndex=FileSize; FillIndex < PaddedEncryptedFileSize && InFile.bNeedEncryption; ++FillIndex)
+				{
+					// Fill the trailing buffer with random bytes from file
+					InOutPersistentBuffer[FillIndex] = InOutPersistentBuffer[rand()%FileSize];
+				}
+
+				//Encrypt the buffer before writing it to disk
+				FAES::EncryptData(InOutPersistentBuffer, PaddedEncryptedFileSize, InEncryptionKey);
+				// Update the size to be written
+				OutSizeToWrite = PaddedEncryptedFileSize;
+				OutNewEntry.Info.bEncrypted = true;
+			}
+
+			// Calculate the buffer hash value
+			FSHA1::HashBuffer(InOutPersistentBuffer,FileSize,OutNewEntry.Info.Hash);			
+			OutDataToWrite = InOutPersistentBuffer;
+		}
+	}
+	return bFileExists;
+}
+
+void FinalizeCopyCompressedFileToPak(FArchive& InPak, const FCompressedFileBuffer& CompressedFile, FPakEntryPair& OutNewEntry)
+{
+	check(CompressedFile.TotalCompressedSize != 0);
+
+	check(OutNewEntry.Info.CompressionBlocks.Num() == CompressedFile.CompressedBlocks.Num());
+	check(OutNewEntry.Info.CompressionMethod == CompressedFile.FileCompressionMethod);
+
+	int64 TellPos = OutNewEntry.Info.GetSerializedSize(FPakInfo::PakFile_Version_Latest);
+	const TArray<FPakCompressedBlock>& Blocks = CompressedFile.CompressedBlocks;
+	for (int32 BlockIndex = 0, BlockCount = Blocks.Num(); BlockIndex < BlockCount; ++BlockIndex)
+	{
+		OutNewEntry.Info.CompressionBlocks[BlockIndex].CompressedStart = Blocks[BlockIndex].CompressedStart + TellPos;
+		OutNewEntry.Info.CompressionBlocks[BlockIndex].CompressedEnd = Blocks[BlockIndex].CompressedEnd + TellPos;
+	}
+}
+
+bool PrepareCopyCompressedFileToPak(const FString& InMountPoint, const FPakInputPair& InFile, const FCompressedFileBuffer& CompressedFile, FPakEntryPair& OutNewEntry, uint8*& OutDataToWrite, int64& OutSizeToWrite, const FAES::FAESKey& InEncryptionKey)
+{
+	if (CompressedFile.TotalCompressedSize == 0)
+	{
+		return false;
+	}
+
+	OutNewEntry.Info.CompressionMethod = CompressedFile.FileCompressionMethod;
+	OutNewEntry.Info.CompressionBlocks.AddZeroed(CompressedFile.CompressedBlocks.Num());
+
+	if (InFile.bNeedEncryption && InEncryptionKey.IsValid())
+	{
+		FAES::EncryptData(CompressedFile.CompressedBuffer.Get(), CompressedFile.TotalCompressedSize, InEncryptionKey);
+	}
+
+	//Hash the final buffer thats written
+	FSHA1 Hash;
+	Hash.Update(CompressedFile.CompressedBuffer.Get(), CompressedFile.TotalCompressedSize);
+	Hash.Final();
+
+	// Update file size & Hash
+	OutNewEntry.Info.CompressionBlockSize = CompressedFile.FileCompressionBlockSize;
+	OutNewEntry.Info.UncompressedSize = CompressedFile.OriginalSize;
+	OutNewEntry.Info.Size = CompressedFile.TotalCompressedSize;
+	Hash.GetHash(OutNewEntry.Info.Hash);
+
+	//	Write the header, then the data
+	OutNewEntry.Filename = InFile.Dest.Mid(InMountPoint.Len());
+	OutNewEntry.Info.Offset = 0; // Don't serialize offsets here.
+	OutNewEntry.Info.bEncrypted = InFile.bNeedEncryption && InEncryptionKey.IsValid();
+	OutSizeToWrite = CompressedFile.TotalCompressedSize;
+	OutDataToWrite = CompressedFile.CompressedBuffer.Get();
+	//OutNewEntry.Info.Serialize(InPak,FPakInfo::PakFile_Version_Latest);	
+	//InPak.Serialize(CompressedFile.CompressedBuffer.Get(), CompressedFile.TotalCompressedSize);
+
+	return true;
+}
+
+void ProcessOrderFile(int32 ArgC, TCHAR* ArgV[], TMap<FString, uint64>& OrderMap)
+{
+	// List of all items to add to pak file
+	FString ResponseFile;
+	if (FParse::Value(FCommandLine::Get(), TEXT("-order="), ResponseFile))
+	{
+		FString Text;
+		UE_LOG(LogPakFile, Display, TEXT("Loading pak order file %s..."), *ResponseFile);
+		if (FFileHelper::LoadFileToString(Text, *ResponseFile))
+		{
+			// Read all lines
+			TArray<FString> Lines;
+			Text.ParseIntoArray(Lines, TEXT("\n"), true);
+			for (int32 EntryIndex = 0; EntryIndex < Lines.Num(); EntryIndex++)
+			{
+				Lines[EntryIndex].ReplaceInline(TEXT("\r"), TEXT(""));
+				Lines[EntryIndex].ReplaceInline(TEXT("\n"), TEXT(""));
+				int32 OpenOrderNumber = EntryIndex;
+				if (Lines[EntryIndex].FindLastChar('"', OpenOrderNumber))
+				{
+					FString ReadNum = Lines[EntryIndex].RightChop(OpenOrderNumber+1);
+					Lines[EntryIndex] = Lines[EntryIndex].Left(OpenOrderNumber+1);
+					ReadNum.TrimStartInline();
+					if (ReadNum.IsNumeric())
+					{
+						OpenOrderNumber = FCString::Atoi(*ReadNum);
+					}
+				}
+				Lines[EntryIndex] = Lines[EntryIndex].TrimQuotes();
+				FString Path=FString::Printf(TEXT("%s"), *Lines[EntryIndex]);
+				FPaths::NormalizeFilename(Path);
+				Path = Path.ToLower();
+#if 0
+				if (Path.EndsWith("uexp"))
+				{
+					OpenOrderNumber += (1 << 29);
+				}
+				if (Path.EndsWith("ubulk"))
+				{
+					OpenOrderNumber += (1 << 30);
+				}
+#endif
+				OrderMap.Add(Path, OpenOrderNumber);
+			}
+			UE_LOG(LogPakFile, Display, TEXT("Finished loading pak order file %s."), *ResponseFile);
+		}
+		else 
+		{
+			UE_LOG(LogPakFile, Display, TEXT("Unable to load pak order file %s."), *ResponseFile);
+		}
+	}
+}
+
+static void CommandLineParseHelper(const TCHAR* InCmdLine, TArray<FString>& Tokens, TArray<FString>& Switches)
+{
+	FString NextToken;
+	while(FParse::Token(InCmdLine,NextToken,false))
+	{
+		if((**NextToken == TCHAR('-')))
+		{
+			new(Switches)FString(NextToken.Mid(1));
+		}
+		else
+		{
+			new(Tokens)FString(NextToken);
+		}
+	}
+}
+
+void ProcessCommandLine(int32 ArgC, TCHAR* ArgV[], TArray<FPakInputPair>& Entries, FPakCommandLineParameters& CmdLineParameters)
+{
+	// List of all items to add to pak file
+	FString ResponseFile;
+	FString ClusterSizeString;
+
+	if (FParse::Value(FCommandLine::Get(), TEXT("-blocksize="), ClusterSizeString) && 
+		FParse::Value(FCommandLine::Get(), TEXT("-blocksize="), CmdLineParameters.FileSystemBlockSize))
+	{
+		if (ClusterSizeString.EndsWith(TEXT("MB")))
+		{
+			CmdLineParameters.FileSystemBlockSize *= 1024*1024;
+		}
+		else if (ClusterSizeString.EndsWith(TEXT("KB")))
+		{
+			CmdLineParameters.FileSystemBlockSize *= 1024;
+		}
+	}
+	else
+	{
+		CmdLineParameters.FileSystemBlockSize = 0;
+	}
+
+	FString CompBlockSizeString;
+	if (FParse::Value(FCommandLine::Get(), TEXT("-compressionblocksize="), CompBlockSizeString) &&
+		FParse::Value(FCommandLine::Get(), TEXT("-compressionblocksize="), CmdLineParameters.CompressionBlockSize))
+	{
+		if (CompBlockSizeString.EndsWith(TEXT("MB")))
+		{
+			CmdLineParameters.CompressionBlockSize *= 1024 * 1024;
+		}
+		else if (CompBlockSizeString.EndsWith(TEXT("KB")))
+		{
+			CmdLineParameters.CompressionBlockSize *= 1024;
+		}
+	}
+
+	if (!FParse::Value(FCommandLine::Get(), TEXT("-bitwindow="), CmdLineParameters.CompressionBitWindow))
+	{
+		CmdLineParameters.CompressionBitWindow = DEFAULT_ZLIB_BIT_WINDOW;
+	}
+
+	if (!FParse::Value(FCommandLine::Get(), TEXT("-patchpaddingalign="), CmdLineParameters.PatchFilePadAlign))
+	{
+		CmdLineParameters.PatchFilePadAlign = 0;
+	}
+
+	if (FParse::Param(FCommandLine::Get(), TEXT("encryptindex")))
+	{
+		CmdLineParameters.EncryptIndex = true;
+	}
+
+	FString CompressorFileName;
+	if (FParse::Value(FCommandLine::Get(), TEXT("compressor="), CompressorFileName))
+	{
+		void* CustomCompressorDll = FPlatformProcess::GetDllHandle(*CompressorFileName);
+		if (CustomCompressorDll == nullptr)
+		{
+			UE_LOG(LogPakFile, Error, TEXT("Unable to load custom compressor from %s"), *CompressorFileName);
+			return;
+		}
+
+		static const TCHAR* CreateCustomCompressorExport = TEXT("CreateCustomCompressor");
+		typedef ICustomCompressor* (CreateCustomCompressorFunc)(const TCHAR*);
+		CreateCustomCompressorFunc* CreateCustomCompressor = reinterpret_cast<CreateCustomCompressorFunc*>(FPlatformProcess::GetDllExport(CustomCompressorDll, CreateCustomCompressorExport));
+		if (CreateCustomCompressor == nullptr)
+		{
+			UE_LOG(LogPakFile, Error, TEXT("Unable to find exported symbol '%s' in '%s'"), CreateCustomCompressorExport, *CompressorFileName);
+			return;
+		}
+
+		ICustomCompressor* Compressor = (*CreateCustomCompressor)(FCommandLine::Get());
+		if (Compressor == nullptr)
+		{
+			UE_LOG(LogPakFile, Error, TEXT("Failed to create custom compressor from '%s'"), *CompressorFileName);
+			return;
+		}
+
+		IModularFeatures::Get().RegisterModularFeature(CUSTOM_COMPRESSOR_FEATURE_NAME, Compressor);
+		CmdLineParameters.UseCustomCompressor = true;
+	}
+
+	if (FParse::Param(FCommandLine::Get(), TEXT("overrideplatformcompressor")))
+	{
+		CmdLineParameters.OverridePlatformCompressor = true;
+	}
+
+	if (FParse::Value(FCommandLine::Get(), TEXT("-create="), ResponseFile))
+	{
+		TArray<FString> Lines;
+
+		CmdLineParameters.GeneratePatch = FParse::Value(FCommandLine::Get(), TEXT("-generatepatch="), CmdLineParameters.SourcePatchPakFilename);
+
+		bool bCompress = FParse::Param(FCommandLine::Get(), TEXT("compress"));
+		bool bEncrypt = FParse::Param(FCommandLine::Get(), TEXT("encrypt"));
+
+		bool bParseLines = true;
+		if (IFileManager::Get().DirectoryExists(*ResponseFile))
+		{
+			IFileManager::Get().FindFilesRecursive(Lines, *ResponseFile, TEXT("*"), true, false);
+			bParseLines = false;
+		}
+		else
+		{
+			FString Text;
+			UE_LOG(LogPakFile, Display, TEXT("Loading response file %s"), *ResponseFile);
+			if (FFileHelper::LoadFileToString(Text, *ResponseFile))
+			{
+				// Remove all carriage return characters.
+				Text.ReplaceInline(TEXT("\r"), TEXT(""));
+				// Read all lines
+				Text.ParseIntoArray(Lines, TEXT("\n"), true);
+			}
+			else
+			{
+				UE_LOG(LogPakFile, Error, TEXT("Failed to load %s"), *ResponseFile);
+			}
+		}
+
+		for (int32 EntryIndex = 0; EntryIndex < Lines.Num(); EntryIndex++)
+		{
+			TArray<FString> SourceAndDest;
+			TArray<FString> Switches;
+			if (bParseLines)
+			{
+				Lines[EntryIndex].TrimStartInline();
+				CommandLineParseHelper(*Lines[EntryIndex], SourceAndDest, Switches);
+			}
+			else
+			{
+				SourceAndDest.Add(Lines[EntryIndex]);
+			}
+			if( SourceAndDest.Num() == 0)
+			{
+				continue;
+			}
+			FPakInputPair Input;
+
+			Input.Source = SourceAndDest[0];
+			FPaths::NormalizeFilename(Input.Source);
+			if (SourceAndDest.Num() > 1)
+			{
+				Input.Dest = FPaths::GetPath(SourceAndDest[1]);
+			}
+			else
+			{
+				Input.Dest = FPaths::GetPath(Input.Source);
+			}
+			FPaths::NormalizeFilename(Input.Dest);
+			FPakFile::MakeDirectoryFromPath(Input.Dest);
+
+			//check for compression switches
+			for (int32 Index = 0; Index < Switches.Num(); ++Index)
+			{
+				if (Switches[Index] == TEXT("compress"))
+				{
+					Input.bNeedsCompression = true;
+				}
+				if (Switches[Index] == TEXT("encrypt"))
+				{
+					Input.bNeedEncryption = true;
+				}
+			}
+			Input.bNeedsCompression |= bCompress;
+			Input.bNeedEncryption |= bEncrypt;
+
+			UE_LOG(LogPakFile, Log, TEXT("Added file Source: %s Dest: %s"), *Input.Source, *Input.Dest);
+			Entries.Add(Input);
+		}			
+	}
+	else
+	{
+		// Override destination path.
+		FString MountPoint;
+		FParse::Value(FCommandLine::Get(), TEXT("-dest="), MountPoint);
+		FPaths::NormalizeFilename(MountPoint);
+		FPakFile::MakeDirectoryFromPath(MountPoint);
+
+		// Parse comand line params. The first param after the program name is the created pak name
+		for (int32 Index = 2; Index < ArgC; Index++)
+		{
+			// Skip switches and add everything else to the Entries array
+			TCHAR* Param = ArgV[Index];
+			if (Param[0] != '-')
+			{
+				FPakInputPair Input;
+				Input.Source = Param;
+				FPaths::NormalizeFilename(Input.Source);
+				if (MountPoint.Len() > 0)
+				{
+					FString SourceDirectory( FPaths::GetPath(Input.Source) );
+					FPakFile::MakeDirectoryFromPath(SourceDirectory);
+					Input.Dest = Input.Source.Replace(*SourceDirectory, *MountPoint, ESearchCase::IgnoreCase);
+				}
+				else
+				{
+					Input.Dest = FPaths::GetPath(Input.Source);
+					FPakFile::MakeDirectoryFromPath(Input.Dest);
+				}
+				FPaths::NormalizeFilename(Input.Dest);
+				Entries.Add(Input);
+			}
+		}
+	}
+	UE_LOG(LogPakFile, Display, TEXT("Added %d entries to add to pak file."), Entries.Num());
+}
+
+void CollectFilesToAdd(TArray<FPakInputPair>& OutFilesToAdd, const TArray<FPakInputPair>& InEntries, const TMap<FString, uint64>& OrderMap)
+{
+	UE_LOG(LogPakFile, Display, TEXT("Collecting files to add to pak file..."));
+	const double StartTime = FPlatformTime::Seconds();
+
+	// Start collecting files
+	TSet<FString> AddedFiles;	
+	for (int32 Index = 0; Index < InEntries.Num(); Index++)
+	{
+		const FPakInputPair& Input = InEntries[Index];
+		const FString& Source = Input.Source;
+		bool bCompression = Input.bNeedsCompression;
+		bool bEncryption = Input.bNeedEncryption;
+
+
+		FString Filename = FPaths::GetCleanFilename(Source);
+		FString Directory = FPaths::GetPath(Source);
+		FPaths::MakeStandardFilename(Directory);
+		FPakFile::MakeDirectoryFromPath(Directory);
+
+		if (Filename.IsEmpty())
+		{
+			Filename = TEXT("*.*");
+		}
+		if ( Filename.Contains(TEXT("*")) )
+		{
+			// Add multiple files
+			TArray<FString> FoundFiles;
+			IFileManager::Get().FindFilesRecursive(FoundFiles, *Directory, *Filename, true, false);
+
+			for (int32 FileIndex = 0; FileIndex < FoundFiles.Num(); FileIndex++)
+			{
+				FPakInputPair FileInput;
+				FileInput.Source = FoundFiles[FileIndex];
+				FPaths::MakeStandardFilename(FileInput.Source);
+				FileInput.Dest = FileInput.Source.Replace(*Directory, *Input.Dest, ESearchCase::IgnoreCase);
+				const uint64* FoundOrder = OrderMap.Find(FileInput.Dest.ToLower());
+				if(FoundOrder)
+				{
+					FileInput.SuggestedOrder = *FoundOrder;
+				}
+				else
+				{
+					// we will put all unordered files at 1 << 28 so that they are before any uexp or ubulk files we assign orders to here
+					FileInput.SuggestedOrder = (1 << 28);
+					// if this is a cook order or an old order it will not have uexp files in it, so we put those in the same relative order after all of the normal files, but before any ubulk files
+					if (FileInput.Dest.EndsWith(TEXT("uexp")) || FileInput.Dest.EndsWith(TEXT("ubulk")))
+					{
+						FoundOrder = OrderMap.Find(FPaths::GetBaseFilename(FileInput.Dest.ToLower(), false) + TEXT(".uasset"));
+						if (!FoundOrder)
+						{
+							FoundOrder = OrderMap.Find(FPaths::GetBaseFilename(FileInput.Dest.ToLower(), false) + TEXT(".umap"));
+						}
+						if (FileInput.Dest.EndsWith(TEXT("uexp")))
+						{
+							FileInput.SuggestedOrder = (FoundOrder ? *FoundOrder : 0) + (1 << 29);
+						}
+						else
+						{
+							FileInput.SuggestedOrder = (FoundOrder ? *FoundOrder : 0) + (1 << 30);
+						}
+					}
+				}
+				FileInput.bNeedsCompression = bCompression;
+				FileInput.bNeedEncryption = bEncryption;
+				if (!AddedFiles.Contains(FileInput.Source))
+				{
+					OutFilesToAdd.Add(FileInput);
+					AddedFiles.Add(FileInput.Source);
+				}
+				else
+				{
+					int32 FoundIndex;
+					OutFilesToAdd.Find(FileInput,FoundIndex);
+					OutFilesToAdd[FoundIndex].bNeedEncryption |= bEncryption;
+					OutFilesToAdd[FoundIndex].bNeedsCompression |= bCompression;
+					OutFilesToAdd[FoundIndex].SuggestedOrder = FMath::Min<uint64>(OutFilesToAdd[FoundIndex].SuggestedOrder, FileInput.SuggestedOrder);
+				}
+			}
+		}
+		else
+		{
+			// Add single file
+			FPakInputPair FileInput;
+			FileInput.Source = Input.Source;
+			FPaths::MakeStandardFilename(FileInput.Source);
+			FileInput.Dest = FileInput.Source.Replace(*Directory, *Input.Dest, ESearchCase::IgnoreCase);
+			const uint64* FoundOrder = OrderMap.Find(FileInput.Dest.ToLower());
+			if (FoundOrder)
+			{
+				FileInput.SuggestedOrder = *FoundOrder;
+			}
+			FileInput.bNeedEncryption = bEncryption;
+			FileInput.bNeedsCompression = bCompression;
+
+			if (AddedFiles.Contains(FileInput.Source))
+			{
+				int32 FoundIndex;
+				OutFilesToAdd.Find(FileInput, FoundIndex);
+				OutFilesToAdd[FoundIndex].bNeedEncryption |= bEncryption;
+				OutFilesToAdd[FoundIndex].bNeedsCompression |= bCompression;
+				OutFilesToAdd[FoundIndex].SuggestedOrder = FMath::Min<uint64>(OutFilesToAdd[FoundIndex].SuggestedOrder, FileInput.SuggestedOrder);
+			}
+			else
+			{
+				OutFilesToAdd.Add(FileInput);
+				AddedFiles.Add(FileInput.Source);
+			}
+		}
+	}
+
+	// Sort by suggested order then alphabetically
+	struct FInputPairSort
+	{
+		FORCEINLINE bool operator()(const FPakInputPair& A, const FPakInputPair& B) const
+		{
+			return A.SuggestedOrder == B.SuggestedOrder ? A.Dest < B.Dest : A.SuggestedOrder < B.SuggestedOrder;
+		}
+	};
+	OutFilesToAdd.Sort(FInputPairSort());
+	UE_LOG(LogPakFile, Display, TEXT("Collected %d files in %.2lfs."), OutFilesToAdd.Num(), FPlatformTime::Seconds() - StartTime);
+}
+
+bool BufferedCopyFile(FArchive& Dest, FArchive& Source, const FPakEntry& Entry, void* Buffer, int64 BufferSize, const FAES::FAESKey& Key)
+{	
+	// Align down
+	BufferSize = BufferSize & ~(FAES::AESBlockSize-1);
+	int64 RemainingSizeToCopy = Entry.Size;
+	while (RemainingSizeToCopy > 0)
+	{
+		const int64 SizeToCopy = FMath::Min(BufferSize, RemainingSizeToCopy);
+		// If file is encrypted so we need to account for padding
+		int64 SizeToRead = Entry.bEncrypted ? Align(SizeToCopy,FAES::AESBlockSize) : SizeToCopy;
+
+		Source.Serialize(Buffer,SizeToRead);
+		if (Entry.bEncrypted)
+		{
+			FAES::DecryptData((uint8*)Buffer, SizeToRead, Key);
+		}
+		Dest.Serialize(Buffer, SizeToCopy);
+		RemainingSizeToCopy -= SizeToRead;
+	}
+	return true;
+}
+
+bool UncompressCopyFile(FArchive& Dest, FArchive& Source, const FPakEntry& Entry, uint8*& PersistentBuffer, int64& BufferSize, const FAES::FAESKey& Key, const FPakFile& PakFile)
+{
+	if (Entry.UncompressedSize == 0)
+	{
+		return false;
+	}
+
+	int64 WorkingSize = Entry.CompressionBlockSize;
+	int32 MaxCompressionBlockSize = FCompression::CompressMemoryBound((ECompressionFlags)Entry.CompressionMethod, WorkingSize);
+	WorkingSize += MaxCompressionBlockSize;
+	if (BufferSize < WorkingSize)
+	{
+		PersistentBuffer = (uint8*)FMemory::Realloc(PersistentBuffer, WorkingSize);
+		BufferSize = WorkingSize;
+	}
+
+	uint8* UncompressedBuffer = PersistentBuffer+MaxCompressionBlockSize;
+
+	for (uint32 BlockIndex=0, BlockIndexNum=Entry.CompressionBlocks.Num(); BlockIndex < BlockIndexNum; ++BlockIndex)
+	{
+		uint32 CompressedBlockSize = Entry.CompressionBlocks[BlockIndex].CompressedEnd - Entry.CompressionBlocks[BlockIndex].CompressedStart;
+		uint32 UncompressedBlockSize = (uint32)FMath::Min<int64>(Entry.UncompressedSize - Entry.CompressionBlockSize*BlockIndex, Entry.CompressionBlockSize);
+		Source.Seek(Entry.CompressionBlocks[BlockIndex].CompressedStart + (PakFile.GetInfo().HasRelativeCompressedChunkOffsets() ? Entry.Offset : 0));
+		uint32 SizeToRead = Entry.bEncrypted ? Align(CompressedBlockSize, FAES::AESBlockSize) : CompressedBlockSize;
+		Source.Serialize(PersistentBuffer, SizeToRead);
+
+		if (Entry.bEncrypted)
+		{
+			FAES::DecryptData(PersistentBuffer, SizeToRead, Key);
+		}
+
+		if(!FCompression::UncompressMemory((ECompressionFlags)Entry.CompressionMethod,UncompressedBuffer,UncompressedBlockSize,PersistentBuffer,CompressedBlockSize))
+		{
+			return false;
+		}
+		Dest.Serialize(UncompressedBuffer,UncompressedBlockSize);
+	}
+
+	return true;
+}
+
+TEncryptionInt ParseEncryptionIntFromJson(TSharedPtr<FJsonObject> InObj, const TCHAR* InName)
+{
+	FString Base64;
+	if (InObj->TryGetStringField(InName, Base64))
+	{
+		TArray<uint8> Bytes;
+		FBase64::Decode(Base64, Bytes);
+		check(Bytes.Num() == sizeof(TEncryptionInt));
+		return TEncryptionInt((uint32*)&Bytes[0]);
+	}
+	else
+	{
+		return TEncryptionInt();
+	}
+}
+
+void PrepareEncryptionAndSigningKeysFromCryptoKeyCache(const FString& InFilename, FKeyPair& OutSigningKey, FAES::FAESKey& OutAESKey)
+{
+	FArchive* File = IFileManager::Get().CreateFileReader(*InFilename);
+	TSharedPtr<FJsonObject> RootObject;
+	TSharedRef<TJsonReader<char>> Reader = TJsonReaderFactory<char>::Create(File);
+	if (FJsonSerializer::Deserialize(Reader, RootObject))
+	{
+		const bool bDataCryptoRequired = RootObject->GetBoolField("bDataCryptoRequired");
+
+		if (bDataCryptoRequired)
+		{
+			const TSharedPtr<FJsonObject>* EncryptionKeyObject;
+			if (RootObject->TryGetObjectField(TEXT("EncryptionKey"), EncryptionKeyObject))
+			{
+				FString EncryptionKeyBase64;
+				if ((*EncryptionKeyObject)->TryGetStringField(TEXT("Key"), EncryptionKeyBase64))
+				{
+					if (EncryptionKeyBase64.Len() > 0)
+					{
+						TArray<uint8> Key;
+						FBase64::Decode(EncryptionKeyBase64, Key);
+						check(Key.Num() == sizeof(FAES::FAESKey::Key));
+						FMemory::Memcpy(OutAESKey.Key, &Key[0], sizeof(FAES::FAESKey::Key));
+					}
+				}
+			}
+
+			const TSharedPtr<FJsonObject>* SigningKey;
+			if (RootObject->TryGetObjectField(TEXT("SigningKey"), SigningKey))
+			{
+				TSharedPtr<FJsonObject> PublicKey = (*SigningKey)->GetObjectField(TEXT("PublicKey"));
+				TSharedPtr<FJsonObject> PrivateKey = (*SigningKey)->GetObjectField(TEXT("PrivateKey"));
+				OutSigningKey.PublicKey.Exponent = ParseEncryptionIntFromJson(PublicKey, TEXT("Exponent"));
+				OutSigningKey.PublicKey.Modulus = ParseEncryptionIntFromJson(PublicKey, TEXT("Modulus"));
+				OutSigningKey.PrivateKey.Exponent = ParseEncryptionIntFromJson(PrivateKey, TEXT("Exponent"));
+				OutSigningKey.PrivateKey.Modulus = ParseEncryptionIntFromJson(PrivateKey, TEXT("Modulus"));
+				check(OutSigningKey.PublicKey.Modulus == OutSigningKey.PrivateKey.Modulus);
+			}
+		}
+	}
+	delete File;
+}
+
+void PrepareEncryptionAndSigningKeys(FKeyPair& OutSigningKey, FAES::FAESKey& OutAESKey)
+{
+	OutSigningKey.PrivateKey.Exponent.Zero();
+	OutSigningKey.PrivateKey.Modulus.Zero();
+	OutSigningKey.PublicKey.Exponent.Zero();
+	OutSigningKey.PublicKey.Modulus.Zero();
+	OutAESKey.Reset();
+
+	// First, try and parse the keys from a supplied crypto key cache file
+	FString CryptoKeysCacheFilename;
+	if (FParse::Value(FCommandLine::Get(), TEXT("cryptokeys="), CryptoKeysCacheFilename))
+	{
+		UE_LOG(LogPakFile, Display, TEXT("Parsing crypto keys from a crypto key cache file"));
+		PrepareEncryptionAndSigningKeysFromCryptoKeyCache(CryptoKeysCacheFilename, OutSigningKey, OutAESKey);
+	}
+	else if (FParse::Param(FCommandLine::Get(), TEXT("encryptionini")))
+	{
+		FString ProjectDir, EngineDir, Platform;
+
+		if (FParse::Value(FCommandLine::Get(), TEXT("projectdir="), ProjectDir, false)
+			&& FParse::Value(FCommandLine::Get(), TEXT("enginedir="), EngineDir, false)
+			&& FParse::Value(FCommandLine::Get(), TEXT("platform="), Platform, false))
+		{
+			FConfigFile EngineConfig;
+
+			FConfigCacheIni::LoadExternalIniFile(EngineConfig, TEXT("Engine"), *FPaths::Combine(EngineDir, TEXT("Config\\")), *FPaths::Combine(ProjectDir, TEXT("Config/")), true, *Platform);
+			bool bDataCryptoRequired = false;
+			EngineConfig.GetBool(TEXT("PlatformCrypto"), TEXT("PlatformRequiresDataCrypto"), bDataCryptoRequired);
+
+			if (!bDataCryptoRequired)
+			{
+				return;
+			}
+
+			FConfigFile ConfigFile;
+			FConfigCacheIni::LoadExternalIniFile(ConfigFile, TEXT("Crypto"), *FPaths::Combine(EngineDir, TEXT("Config\\")), *FPaths::Combine(ProjectDir, TEXT("Config/")), true, *Platform);
+			bool bSignPak = false;
+			bool bEncryptPakIniFiles = false;
+			bool bEncryptPakIndex = false;
+			bool bEncryptAssets = false;
+			bool bEncryptPak = false;
+
+			if (ConfigFile.Num())
+			{
+				UE_LOG(LogPakFile, Display, TEXT("Using new format crypto.ini files for crypto configuration"));
+
+				static const TCHAR* SectionName = TEXT("/Script/CryptoKeys.CryptoKeysSettings");
+
+				ConfigFile.GetBool(SectionName, TEXT("bEnablePakSigning"), bSignPak);
+				ConfigFile.GetBool(SectionName, TEXT("bEncryptPakIniFiles"), bEncryptPakIniFiles);
+				ConfigFile.GetBool(SectionName, TEXT("bEncryptPakIndex"), bEncryptPakIndex);
+				ConfigFile.GetBool(SectionName, TEXT("bEncryptAssets"), bEncryptAssets);
+				bEncryptPak = bEncryptPakIniFiles || bEncryptPakIndex || bEncryptAssets;
+
+				if (bSignPak)
+				{
+					FString PublicExpBase64, PrivateExpBase64, ModulusBase64;
+					ConfigFile.GetString(SectionName, TEXT("SigningPublicExponent"), PublicExpBase64);
+					ConfigFile.GetString(SectionName, TEXT("SigningPrivateExponent"), PrivateExpBase64);
+					ConfigFile.GetString(SectionName, TEXT("SigningModulus"), ModulusBase64);
+
+					TArray<uint8> PublicExp, PrivateExp, Modulus;
+					FBase64::Decode(PublicExpBase64, PublicExp);
+					FBase64::Decode(PrivateExpBase64, PrivateExp);
+					FBase64::Decode(ModulusBase64, Modulus);
+
+					OutSigningKey.PrivateKey.Exponent = TEncryptionInt((uint32*)&PrivateExp[0]);
+					OutSigningKey.PrivateKey.Modulus = TEncryptionInt((uint32*)&Modulus[0]);
+					OutSigningKey.PublicKey.Exponent = TEncryptionInt((uint32*)&PublicExp[0]);
+					OutSigningKey.PublicKey.Modulus = OutSigningKey.PrivateKey.Modulus;
+
+					UE_LOG(LogPakFile, Display, TEXT("Parsed signature keys from config files."));
+				}
+
+				if (bEncryptPak)
+				{
+					FString EncryptionKeyString;
+					ConfigFile.GetString(SectionName, TEXT("EncryptionKey"), EncryptionKeyString);
+
+					if (EncryptionKeyString.Len() > 0)
+					{
+						TArray<uint8> Key;
+						FBase64::Decode(EncryptionKeyString, Key);
+						check(Key.Num() == sizeof(FAES::FAESKey::Key));
+						FMemory::Memcpy(OutAESKey.Key, &Key[0], sizeof(FAES::FAESKey::Key));
+						UE_LOG(LogPakFile, Display, TEXT("Parsed AES encryption key from config files."));
+					}
+				}
+			}
+			else
+			{
+				static const TCHAR* SectionName = TEXT("Core.Encryption");
+
+				UE_LOG(LogPakFile, Display, TEXT("Using old format encryption.ini files for crypto configuration"));
+
+				FConfigCacheIni::LoadExternalIniFile(ConfigFile, TEXT("Encryption"), *FPaths::Combine(EngineDir, TEXT("Config\\")), *FPaths::Combine(ProjectDir, TEXT("Config/")), true, *Platform);
+				ConfigFile.GetBool(SectionName, TEXT("SignPak"), bSignPak);
+				ConfigFile.GetBool(SectionName, TEXT("EncryptPak"), bEncryptPak);
+
+				if (bSignPak)
+				{
+					FString RSAPublicExp, RSAPrivateExp, RSAModulus;
+					ConfigFile.GetString(SectionName, TEXT("rsa.publicexp"), RSAPublicExp);
+					ConfigFile.GetString(SectionName, TEXT("rsa.privateexp"), RSAPrivateExp);
+					ConfigFile.GetString(SectionName, TEXT("rsa.modulus"), RSAModulus);
+
+					OutSigningKey.PrivateKey.Exponent.Parse(RSAPrivateExp);
+					OutSigningKey.PrivateKey.Modulus.Parse(RSAModulus);
+					OutSigningKey.PublicKey.Exponent.Parse(RSAPublicExp);
+					OutSigningKey.PublicKey.Modulus = OutSigningKey.PrivateKey.Modulus;
+
+					UE_LOG(LogPakFile, Display, TEXT("Parsed signature keys from config files."));
+				}
+
+				if (bEncryptPak)
+				{
+					FString EncryptionKeyString;
+					ConfigFile.GetString(SectionName, TEXT("aes.key"), EncryptionKeyString);
+
+					if (EncryptionKeyString.Len() == 32 && TCString<TCHAR>::IsPureAnsi(*EncryptionKeyString))
+					{
+						for (int32 Index = 0; Index < 32; ++Index)
+						{
+							OutAESKey.Key[Index] = (uint8)EncryptionKeyString[Index];
+						}
+
+						UE_LOG(LogPakFile, Display, TEXT("Parsed AES encryption key from config files."));
+					}
+				}
+			}
+		}
+	}
+	else
+	{
+		UE_LOG(LogPakFile, Display, TEXT("Using command line for crypto configuration"));
+
+		FString EncryptionKeyString;
+		FParse::Value(FCommandLine::Get(), TEXT("aes="), EncryptionKeyString, false);
+
+		if (EncryptionKeyString.Len() > 0)
+		{
+			const uint32 RequiredKeyLength = sizeof(OutAESKey.Key);
+
+			// Error checking
+			if (EncryptionKeyString.Len() < RequiredKeyLength)
+			{
+				UE_LOG(LogPakFile, Fatal, TEXT("AES encryption key must be %d characters long"), RequiredKeyLength);
+			}
+
+			if (EncryptionKeyString.Len() > RequiredKeyLength)
+			{
+				UE_LOG(LogPakFile, Warning, TEXT("AES encryption key is more than %d characters long, so will be truncated!"), RequiredKeyLength);
+				EncryptionKeyString = EncryptionKeyString.Left(RequiredKeyLength);
+			}
+
+			if (!FCString::IsPureAnsi(*EncryptionKeyString))
+			{
+				UE_LOG(LogPakFile, Fatal, TEXT("AES encryption key must be a pure ANSI string!"));
+			}
+
+			ANSICHAR* AsAnsi = TCHAR_TO_ANSI(*EncryptionKeyString);
+			check(TCString<ANSICHAR>::Strlen(AsAnsi) == RequiredKeyLength);
+			FMemory::Memcpy(OutAESKey.Key, AsAnsi, RequiredKeyLength);
+			UE_LOG(LogPakFile, Display, TEXT("Parsed AES encryption key from command line."));
+		}
+
+		FString KeyFilename;
+		if (FParse::Value(FCommandLine::Get(), TEXT("sign="), KeyFilename, false))
+		{
+			if (KeyFilename.StartsWith(TEXT("0x")))
+			{
+				TArray<FString> KeyValueText;
+				int32 NumParts = KeyFilename.ParseIntoArray(KeyValueText, TEXT("+"), true);
+				if (NumParts == 3)
+				{
+					OutSigningKey.PrivateKey.Exponent.Parse(KeyValueText[0]);
+					OutSigningKey.PrivateKey.Modulus.Parse(KeyValueText[1]);
+					OutSigningKey.PublicKey.Exponent.Parse(KeyValueText[2]);
+					OutSigningKey.PublicKey.Modulus = OutSigningKey.PrivateKey.Modulus;
+
+					UE_LOG(LogPakFile, Display, TEXT("Parsed signature keys from command line."));
+				}
+				else
+				{
+					UE_LOG(LogPakFile, Error, TEXT("Expected 3, got %d, when parsing %s"), KeyValueText.Num(), *KeyFilename);
+					OutSigningKey.PrivateKey.Exponent.Zero();
+				}
+			}
+			else if (!ReadKeysFromFile(*KeyFilename, OutSigningKey))
+			{
+				UE_LOG(LogPakFile, Error, TEXT("Unable to load signature keys %s."), *KeyFilename);
+			}
+		}
+	}
+
+	if (OutSigningKey.IsValid())
+	{
+		if (!TestKeys(OutSigningKey))
+		{
+			UE_LOG(LogPakFile, Fatal, TEXT("Pak signing keys are invalid"));
+			OutSigningKey.PrivateKey.Exponent.Zero();
+		}
+	}
+
+	if (OutAESKey.IsValid())
+	{
+		FCoreDelegates::GetPakEncryptionKeyDelegate().BindLambda([OutAESKey](uint8 OutKey[32]) { FMemory::Memcpy(OutKey, OutAESKey.Key, sizeof(OutAESKey.Key)); });
+	}
+}
+
+/**
+ * Creates a pak file writer. This can be a signed writer if the encryption keys are specified in the command line
+ */
+FArchive* CreatePakWriter(const TCHAR* Filename, const FKeyPair& SigningKey)
+{
+	FArchive* Writer = IFileManager::Get().CreateFileWriter(Filename);
+	FString KeyFilename;
+	bool bSigningEnabled = false;
+	
+	if (Writer)
+	{
+		if (SigningKey.IsValid())
+		{
+			UE_LOG(LogPakFile, Display, TEXT("Creating signed pak %s."), Filename);
+			Writer = new FSignedArchiveWriter(*Writer, Filename, SigningKey.PublicKey, SigningKey.PrivateKey);
+		}
+	}
+
+	return Writer;
+}
+
+bool CreatePakFile(const TCHAR* Filename, TArray<FPakInputPair>& FilesToAdd, const FPakCommandLineParameters& CmdLineParameters, const FKeyPair& SigningKey, const FAES::FAESKey& EncryptionKey)
+{	
+	const double StartTime = FPlatformTime::Seconds();
+
+	// Create Pak
+	TUniquePtr<FArchive> PakFileHandle(CreatePakWriter(Filename, SigningKey));
+	if (!PakFileHandle)
+	{
+		UE_LOG(LogPakFile, Error, TEXT("Unable to create pak file \"%s\"."), Filename);
+		return false;
+	}
+
+	FPakInfo Info;
+	Info.bEncryptedIndex = (EncryptionKey.IsValid() && CmdLineParameters.EncryptIndex);
+
+	TArray<FPakEntryPair> Index;
+	FString MountPoint = GetCommonRootPath(FilesToAdd);
+	uint8* ReadBuffer = NULL;
+	int64 BufferSize = 0;
+	ECompressionFlags CompressionMethod = COMPRESS_None;
+	FCompressedFileBuffer CompressedFileBuffer;
+
+	uint8* PaddingBuffer = nullptr;
+	int64 PaddingBufferSize = 0;
+	if (CmdLineParameters.PatchFilePadAlign > 0)
+	{
+		PaddingBufferSize = CmdLineParameters.PatchFilePadAlign;
+		PaddingBuffer = (uint8*)FMemory::Malloc(PaddingBufferSize);
+		FMemory::Memset(PaddingBuffer, 0, PaddingBufferSize);
+	}
+
+	// Some platforms provide patch download size reduction by diffing the patch files.  However, they often operate on specific block
+	// sizes when dealing with new data within the file.  Pad files out to the given alignment to work with these systems more nicely.
+	// We also want to combine smaller files into the same padding size block so we don't waste as much space. i.e. grouping 64 1k files together
+	// rather than padding each out to 64k.
+	const uint32 RequiredPatchPadding = CmdLineParameters.PatchFilePadAlign;
+
+	uint64 ContiguousTotalSizeSmallerThanBlockSize = 0;
+	uint64 ContiguousFilesSmallerThanBlockSize = 0;
+
+	uint64 TotalUncompressedSize = 0;
+	uint64 TotalCompressedSize = 0;
+
+	uint64 TotalRequestedEncryptedFiles = 0;
+	uint64 TotalEncryptedFiles = 0;
+	uint64 TotalEncryptedDataSize = 0;
+
+	for (int32 FileIndex = 0; FileIndex < FilesToAdd.Num(); FileIndex++)
+	{
+		bool bIsUAssetUExpPairUAsset = false;
+		bool bIsUAssetUExpPairUExp = false;
+
+		if (FileIndex)
+		{
+			if (FPaths::GetBaseFilename(FilesToAdd[FileIndex - 1].Dest, false) == FPaths::GetBaseFilename(FilesToAdd[FileIndex].Dest, false) &&
+				FPaths::GetExtension(FilesToAdd[FileIndex - 1].Dest, true) == TEXT(".uasset") && 
+				FPaths::GetExtension(FilesToAdd[FileIndex].Dest, true) == TEXT(".uexp")
+				)
+			{
+				bIsUAssetUExpPairUExp = true;
+			}
+		}
+		if (!bIsUAssetUExpPairUExp && FileIndex + 1 < FilesToAdd.Num())
+		{
+			if (FPaths::GetBaseFilename(FilesToAdd[FileIndex].Dest, false) == FPaths::GetBaseFilename(FilesToAdd[FileIndex + 1].Dest, false) &&
+				FPaths::GetExtension(FilesToAdd[FileIndex].Dest, true) == TEXT(".uasset") &&
+				FPaths::GetExtension(FilesToAdd[FileIndex + 1].Dest, true) == TEXT(".uexp")
+				)
+			{
+				bIsUAssetUExpPairUAsset = true;
+			}
+		}
+
+		//  Remember the offset but don't serialize it with the entry header.
+		int64 NewEntryOffset = PakFileHandle->Tell();
+		FPakEntryPair NewEntry;
+
+		//check if this file requested to be compression
+		int64 OriginalFileSize = IFileManager::Get().FileSize(*FilesToAdd[FileIndex].Source);
+		int64 RealFileSize = OriginalFileSize + NewEntry.Info.GetSerializedSize(FPakInfo::PakFile_Version_Latest);
+		if (FilesToAdd[FileIndex].bNeedsCompression && OriginalFileSize > 0)
+		{
+			if (CmdLineParameters.UseCustomCompressor)
+			{
+				CompressionMethod = COMPRESS_Custom;
+			}
+			else
+			{
+				CompressionMethod = COMPRESS_Default;
+			}
+
+			if (CmdLineParameters.OverridePlatformCompressor)
+			{
+				CompressionMethod = (ECompressionFlags)(CompressionMethod | COMPRESS_OverridePlatform);
+			}
+		}
+		else
+		{
+			CompressionMethod = COMPRESS_None;
+		}
+
+		if (CompressionMethod != COMPRESS_None)
+		{
+			if (CompressedFileBuffer.CompressFileToWorkingBuffer(FilesToAdd[FileIndex], ReadBuffer, BufferSize, CompressionMethod, CmdLineParameters.CompressionBlockSize, CmdLineParameters.CompressionBitWindow))
+			{
+				// Check the compression ratio, if it's too low just store uncompressed. Also take into account read size
+				// if we still save 64KB it's probably worthwhile compressing, as that saves a file read operation in the runtime.
+				// TODO: drive this threashold from the command line
+				float PercentLess = ((float)CompressedFileBuffer.TotalCompressedSize / (OriginalFileSize / 100.f));
+				if (PercentLess > 90.f && (OriginalFileSize - CompressedFileBuffer.TotalCompressedSize) < 65536)
+				{
+					CompressionMethod = COMPRESS_None;
+				}
+				else
+				{
+					NewEntry.Info.CompressionMethod = CompressionMethod;
+					NewEntry.Info.CompressionBlocks.AddUninitialized(CompressedFileBuffer.CompressedBlocks.Num());
+					RealFileSize = CompressedFileBuffer.TotalCompressedSize + NewEntry.Info.GetSerializedSize(FPakInfo::PakFile_Version_Latest);
+					NewEntry.Info.CompressionBlocks.Reset();
+				}
+			}
+			else
+			{
+				// Compression failed. Include file uncompressed and warn the user.
+				UE_LOG(LogPakFile, Warning, TEXT("File \"%s\" failed compression. File will be saved uncompressed."), *FilesToAdd[FileIndex].Source);
+				CompressionMethod = COMPRESS_None;
+			}
+		}
+
+		// Account for file system block size, which is a boundary we want to avoid crossing.
+		if (!bIsUAssetUExpPairUExp && // don't split uexp / uasset pairs
+			CmdLineParameters.FileSystemBlockSize > 0 && OriginalFileSize != INDEX_NONE && RealFileSize <= CmdLineParameters.FileSystemBlockSize)
+		{
+			if ((NewEntryOffset / CmdLineParameters.FileSystemBlockSize) != ((NewEntryOffset + RealFileSize) / CmdLineParameters.FileSystemBlockSize))
+			{
+				//File crosses a block boundary, so align it to the beginning of the next boundary
+				int64 OldOffset = NewEntryOffset;
+				NewEntryOffset = AlignArbitrary(NewEntryOffset, CmdLineParameters.FileSystemBlockSize);
+				int64 PaddingRequired = NewEntryOffset - OldOffset;
+
+				if (PaddingRequired > 0)
+				{
+					// If we don't already have a padding buffer, create one
+					if (PaddingBuffer == nullptr)
+					{
+						PaddingBufferSize = 64 * 1024;
+						PaddingBuffer = (uint8*)FMemory::Malloc(PaddingBufferSize);
+						FMemory::Memset(PaddingBuffer, 0, PaddingBufferSize);
+					}
+
+					UE_LOG(LogPakFile, Verbose, TEXT("%14llu - %14llu : %14llu padding."), PakFileHandle->Tell(), PakFileHandle->Tell() + PaddingRequired, PaddingRequired);
+					while (PaddingRequired > 0)
+					{
+						int64 AmountToWrite = FMath::Min(PaddingRequired, PaddingBufferSize);
+						PakFileHandle->Serialize(PaddingBuffer, AmountToWrite);
+						PaddingRequired -= AmountToWrite;
+					}
+
+					check(PakFileHandle->Tell() == NewEntryOffset);
+				}
+			}
+		}
+
+		bool bCopiedToPak;
+		int64 SizeToWrite = 0;
+		uint8* DataToWrite = nullptr;
+		if (FilesToAdd[FileIndex].bNeedsCompression && CompressionMethod != COMPRESS_None)
+		{
+			bCopiedToPak = PrepareCopyCompressedFileToPak(MountPoint, FilesToAdd[FileIndex], CompressedFileBuffer, NewEntry, DataToWrite, SizeToWrite, EncryptionKey);
+			DataToWrite = CompressedFileBuffer.CompressedBuffer.Get();
+		}
+		else
+		{
+			bCopiedToPak = PrepareCopyFileToPak(MountPoint, FilesToAdd[FileIndex], ReadBuffer, BufferSize, NewEntry, DataToWrite, SizeToWrite, EncryptionKey);
+			DataToWrite = ReadBuffer;
+		}		
+
+		int64 TotalSizeToWrite = SizeToWrite + NewEntry.Info.GetSerializedSize(FPakInfo::PakFile_Version_Latest);
+		if (bCopiedToPak)
+		{
+			
+			if (RequiredPatchPadding > 0)
+			{
+				//if the next file is going to cross a patch-block boundary then pad out the current set of files with 0's
+				//and align the next file up.
+				bool bCrossesBoundary = AlignArbitrary(NewEntryOffset, RequiredPatchPadding) != AlignArbitrary(NewEntryOffset + TotalSizeToWrite - 1, RequiredPatchPadding);
+				bool bPatchPadded = false;
+				if (!bIsUAssetUExpPairUExp) // never patch-pad the uexp of a uasset/uexp pair
+				{
+					bool bPairProbablyCrossesBoundary = false; // we don't consider compression because we have not compressed the uexp yet.
+					if (bIsUAssetUExpPairUAsset)
+					{
+						int64 UExpFileSize = IFileManager::Get().FileSize(*FilesToAdd[FileIndex + 1].Source) / 2; // assume 50% compression
+						bPairProbablyCrossesBoundary = AlignArbitrary(NewEntryOffset, RequiredPatchPadding) != AlignArbitrary(NewEntryOffset + TotalSizeToWrite + UExpFileSize - 1, RequiredPatchPadding);
+					}
+					if (TotalSizeToWrite >= RequiredPatchPadding || // if it exactly the padding size and by luck does not cross a boundary, we still consider it "over" because it can't be packed with anything else
+						bCrossesBoundary || bPairProbablyCrossesBoundary)
+					{
+						NewEntryOffset = AlignArbitrary(NewEntryOffset, RequiredPatchPadding);
+						int64 CurrentLoc = PakFileHandle->Tell();
+						int64 PaddingSize = NewEntryOffset - CurrentLoc;
+						check(PaddingSize >= 0);
+						if (PaddingSize)
+						{
+							UE_LOG(LogPakFile, Verbose, TEXT("%14llu - %14llu : %14llu patch padding."), PakFileHandle->Tell(), PakFileHandle->Tell() + PaddingSize, PaddingSize);
+							check(PaddingSize <= PaddingBufferSize);
+
+							//have to pad manually with 0's.  File locations skipped by Seek and never written are uninitialized which would defeat the whole purpose
+							//of padding for certain platforms patch diffing systems.
+							PakFileHandle->Serialize(PaddingBuffer, PaddingSize);
+						}
+						check(PakFileHandle->Tell() == NewEntryOffset);
+						bPatchPadded = true;
+					}
+				}
+
+				//if the current file is bigger than a patch block then we will always have to pad out the previous files.
+				//if there were a large set of contiguous small files behind us then this will be the natural stopping point for a possible pathalogical patching case where growth in the small files causes a cascade 
+				//to dirty up all the blocks prior to this one.  If this could happen let's warn about it.
+				if (bPatchPadded ||
+					FileIndex + 1 == FilesToAdd.Num()) // also check the last file, this won't work perfectly if we don't end up adding the last file for some reason
+				{
+					const uint64 ContiguousGroupedFilePatchWarningThreshhold = 50 * 1024 * 1024;
+					if (ContiguousTotalSizeSmallerThanBlockSize > ContiguousGroupedFilePatchWarningThreshhold)
+					{
+						UE_LOG(LogPakFile, Display, TEXT("%i small files (%i) totaling %llu contiguous bytes found before first 'large' file.  Changes to any of these files could cause the whole group to be 'dirty' in a per-file binary diff based patching system."), ContiguousFilesSmallerThanBlockSize, RequiredPatchPadding, ContiguousTotalSizeSmallerThanBlockSize);
+					}
+					ContiguousTotalSizeSmallerThanBlockSize = 0;
+					ContiguousFilesSmallerThanBlockSize = 0;
+				}
+				else
+				{
+					ContiguousTotalSizeSmallerThanBlockSize += TotalSizeToWrite;
+					ContiguousFilesSmallerThanBlockSize++;				
+				}
+			}
+			if (FilesToAdd[FileIndex].bNeedsCompression && CompressionMethod != COMPRESS_None)
+			{
+				FinalizeCopyCompressedFileToPak(*PakFileHandle, CompressedFileBuffer, NewEntry);
+			}
+
+			// Write to file
+			int64 Offset = PakFileHandle->Tell();
+			NewEntry.Info.Serialize(*PakFileHandle, FPakInfo::PakFile_Version_Latest);
+			PakFileHandle->Serialize(DataToWrite, SizeToWrite);	
+			int64 EndOffset = PakFileHandle->Tell();
+
+			UE_LOG(LogPakFile, Verbose, TEXT("%14llu - %14llu : %14llu header+file %s."), Offset, EndOffset, EndOffset - Offset, *NewEntry.Filename);
+
+			// Update offset now and store it in the index (and only in index)
+			NewEntry.Info.Offset = NewEntryOffset;
+			Index.Add(NewEntry);
+			const TCHAR* EncryptedString = TEXT("");
+
+			if (FilesToAdd[FileIndex].bNeedEncryption)
+			{
+				TotalRequestedEncryptedFiles++;
+
+				if (EncryptionKey.IsValid())
+				{
+					TotalEncryptedFiles++;
+					TotalEncryptedDataSize += SizeToWrite;
+					EncryptedString = TEXT("encrypted ");
+				}
+			}
+
+			if (FilesToAdd[FileIndex].bNeedsCompression && CompressionMethod != COMPRESS_None)
+			{
+				TotalCompressedSize += NewEntry.Info.Size;
+				TotalUncompressedSize += NewEntry.Info.UncompressedSize;
+				float PercentLess = ((float)NewEntry.Info.Size / (NewEntry.Info.UncompressedSize / 100.f));
+				if (FilesToAdd[FileIndex].SuggestedOrder < MAX_uint64)
+				{
+					UE_LOG(LogPakFile, Log, TEXT("Added compressed %sfile \"%s\", %.2f%% of original size. Compressed Size %lld bytes, Original Size %lld bytes (order %llu)."), EncryptedString, *NewEntry.Filename, PercentLess, NewEntry.Info.Size, NewEntry.Info.UncompressedSize, FilesToAdd[FileIndex].SuggestedOrder);
+				}
+				else
+				{
+					UE_LOG(LogPakFile, Log, TEXT("Added compressed %sfile \"%s\", %.2f%% of original size. Compressed Size %lld bytes, Original Size %lld bytes (no order given)."), EncryptedString, *NewEntry.Filename, PercentLess, NewEntry.Info.Size, NewEntry.Info.UncompressedSize);
+				}
+			}
+			else
+			{
+				if (FilesToAdd[FileIndex].SuggestedOrder < MAX_uint64)
+				{
+					UE_LOG(LogPakFile, Log, TEXT("Added %sfile \"%s\", %lld bytes (order %llu)."), EncryptedString, *NewEntry.Filename, NewEntry.Info.Size, FilesToAdd[FileIndex].SuggestedOrder);
+				}
+				else
+				{
+					UE_LOG(LogPakFile, Log, TEXT("Added %sfile \"%s\", %lld bytes (no order given)."), EncryptedString, *NewEntry.Filename, NewEntry.Info.Size);
+				}
+			}
+		}
+		else
+		{
+			UE_LOG(LogPakFile, Warning, TEXT("Missing file \"%s\" will not be added to PAK file."), *FilesToAdd[FileIndex].Source);
+		}
+	}
+
+	FMemory::Free(PaddingBuffer);
+	FMemory::Free(ReadBuffer);
+	ReadBuffer = NULL;
+
+	// Remember IndexOffset
+	Info.IndexOffset = PakFileHandle->Tell();
+
+	// Serialize Pak Index at the end of Pak File
+	TArray<uint8> IndexData;
+	FMemoryWriter IndexWriter(IndexData);
+	IndexWriter.SetByteSwapping(PakFileHandle->ForceByteSwapping());
+	int32 NumEntries = Index.Num();
+	IndexWriter << MountPoint;
+	IndexWriter << NumEntries;
+	for (int32 EntryIndex = 0; EntryIndex < Index.Num(); EntryIndex++)
+	{
+		FPakEntryPair& Entry = Index[EntryIndex];
+		IndexWriter << Entry.Filename;
+		Entry.Info.Serialize(IndexWriter, Info.Version);
+
+		if (RequiredPatchPadding > 0)
+		{
+			int64 EntrySize = Entry.Info.GetSerializedSize(FPakInfo::PakFile_Version_Latest);
+			int64 TotalSizeToWrite = Entry.Info.Size + EntrySize;
+			if (TotalSizeToWrite >= RequiredPatchPadding)
+			{
+				int64 RealStart = Entry.Info.Offset;
+				if ((RealStart % RequiredPatchPadding) != 0 && 
+					!Entry.Filename.EndsWith(TEXT("uexp"))) // these are export sections of larger files and may be packed with uasset/umap and so we don't need a warning here
+				{
+					UE_LOG(LogPakFile, Warning, TEXT("File at offset %lld of size %lld not aligned to patch size %i"), RealStart, Entry.Info.Size, RequiredPatchPadding);
+				}
+			}
+		}
+	}
+
+	if (Info.bEncryptedIndex)
+	{
+		int32 OriginalSize = IndexData.Num();
+		int32 AlignedSize = Align(OriginalSize, FAES::AESBlockSize);
+
+		for (int32 PaddingIndex = IndexData.Num(); PaddingIndex < AlignedSize; ++PaddingIndex)
+		{
+			uint8 Byte = IndexData[PaddingIndex % OriginalSize];
+			IndexData.Add(Byte);
+		}
+	}
+
+	FSHA1::HashBuffer(IndexData.GetData(), IndexData.Num(), Info.IndexHash);
+
+	if (Info.bEncryptedIndex)
+	{
+		FAES::EncryptData(IndexData.GetData(), IndexData.Num(), EncryptionKey);
+		TotalEncryptedDataSize += IndexData.Num();
+	}
+
+	PakFileHandle->Serialize(IndexData.GetData(), IndexData.Num());
+
+	Info.IndexSize = IndexData.Num();
+
+	// Save trailer (offset, size, hash value)
+	Info.Serialize(*PakFileHandle);
+
+	UE_LOG(LogPakFile, Display, TEXT("Added %d files, %lld bytes total, time %.2lfs."), Index.Num(), PakFileHandle->TotalSize(), FPlatformTime::Seconds() - StartTime);
+	if (TotalUncompressedSize)
+	{
+		float PercentLess = ((float)TotalCompressedSize / (TotalUncompressedSize / 100.f));
+		UE_LOG(LogPakFile, Display, TEXT("Compression summary: %.2f%% of original size. Compressed Size %lld bytes, Original Size %lld bytes. "), PercentLess, TotalCompressedSize, TotalUncompressedSize);
+	}
+
+	if (TotalEncryptedDataSize)
+	{
+		UE_LOG(LogPakFile, Display, TEXT("Encryption - ENABLED"));
+		UE_LOG(LogPakFile, Display, TEXT("  Files: %d"), TotalEncryptedFiles);
+
+		if (Info.bEncryptedIndex)
+		{
+			UE_LOG(LogPakFile, Display, TEXT("  Index: Encrypted (%d bytes, %.2fMB)"), Info.IndexSize, (float)Info.IndexSize / 1024.0f / 1024.0f);
+		}
+		else
+		{
+			UE_LOG(LogPakFile, Display, TEXT("  Index: Unencrypted"));
+		}
+		
+
+		UE_LOG(LogPakFile, Display, TEXT("  Total: %d bytes (%.2fMB)"), TotalEncryptedDataSize, (float)TotalEncryptedDataSize / 1024.0f / 1024.0f);
+	}
+	else
+	{
+		UE_LOG(LogPakFile, Display, TEXT("Encryption - DISABLED"));
+	}
+
+	if (TotalEncryptedFiles < TotalRequestedEncryptedFiles)
+	{
+		UE_LOG(LogPakFile, Display, TEXT("%d files requested encryption, but no AES key was supplied! Encryption was skipped for these files"), TotalRequestedEncryptedFiles);
+	}
+
+	PakFileHandle->Close();
+	PakFileHandle.Reset();
+
+	return true;
+}
+
+bool TestPakFile(const TCHAR* Filename)
+{	
+	FPakFile PakFile(Filename, FParse::Param(FCommandLine::Get(), TEXT("signed")));
+	if (PakFile.IsValid())
+	{
+		return PakFile.Check();
+	}
+	else
+	{
+		UE_LOG(LogPakFile, Error, TEXT("Unable to open pak file \"%s\"."), Filename);
+		return false;
+	}
+}
+
+bool ListFilesInPak(const TCHAR * InPakFilename, int64 SizeFilter = 0)
+{
+	FPakFile PakFile(InPakFilename, FParse::Param(FCommandLine::Get(), TEXT("signed")));
+	int32 FileCount = 0;
+	int64 FileSize = 0;
+	int64 FilteredSize = 0;
+
+	if (PakFile.IsValid())
+	{
+		UE_LOG(LogPakFile, Display, TEXT("Mount point %s"), *PakFile.GetMountPoint());
+
+		TArray<FPakFile::FFileIterator> Records;
+
+		for (FPakFile::FFileIterator It(PakFile); It; ++It)
+		{
+			Records.Add(It);
+		}
+
+		struct FOffsetSort
+		{
+			FORCEINLINE bool operator()(const FPakFile::FFileIterator& A, const FPakFile::FFileIterator& B) const
+			{
+				return A.Info().Offset < B.Info().Offset;
+			}
+		};
+
+		Records.Sort(FOffsetSort());
+
+		for (auto It : Records)
+		{
+			const FPakEntry& Entry = It.Info();
+			if (Entry.Size >= SizeFilter)
+			{
+				UE_LOG(LogPakFile, Display, TEXT("\"%s\" offset: %lld, size: %d bytes, sha1: %s."), *It.Filename(), Entry.Offset, Entry.Size, *BytesToHex(Entry.Hash, sizeof(Entry.Hash)));
+				FilteredSize += Entry.Size;
+			}
+			FileSize += Entry.Size;
+			FileCount++;
+		}
+		UE_LOG(LogPakFile, Display, TEXT("%d files (%lld bytes), (%lld filtered bytes)."), FileCount, FileSize, FilteredSize);
+
+		return true;
+	}
+	else
+	{
+		UE_LOG(LogPakFile, Error, TEXT("Unable to open pak file \"%s\"."), InPakFilename);
+		return false;
+	}
+}
+
+struct FFileInfo
+{
+	uint64 FileSize;
+	int32 PatchIndex;
+	uint8 Hash[16];
+};
+
+bool ExtractFilesFromPak(const TCHAR* InPakFilename, TMap<FString, FFileInfo>& InFileHashes, const TCHAR* InDestPath, bool bUseMountPoint, const FAES::FAESKey& InEncryptionKey, TArray<FPakInputPair>* OutEntries = nullptr, TMap<FString, uint64>* OutOrderMap = nullptr)
+{
+	// Gather all patch versions of the requested pak file and run through each separately
+	TArray<FString> PakFileList;
+	FString PakFileDirectory = FPaths::GetPath(InPakFilename);
+	// If file doesn't exist try using it as a search string, it may contain wild cards
+	if (IFileManager::Get().FileExists(InPakFilename))
+	{
+		PakFileList.Add(*FPaths::GetCleanFilename(InPakFilename));
+	}
+	else
+	{
+		IFileManager::Get().FindFiles(PakFileList, *PakFileDirectory, *FPaths::GetCleanFilename(InPakFilename));
+	}
+
+	for (int32 PakFileIndex = 0; PakFileIndex < PakFileList.Num(); PakFileIndex++)
+	{
+		FString PakFilename = PakFileDirectory + "\\" + PakFileList[PakFileIndex];
+		// Gather the pack file index from the filename. The base pak file holds index -1;
+		int32 PakPriority = -1;
+		if (PakFilename.EndsWith("_P.pak"))
+		{
+			FString PakIndexFromFilename = PakFilename.LeftChop(6);
+			int32 PakIndexStart = INDEX_NONE;
+			PakIndexFromFilename.FindLastChar('_', PakIndexStart);
+			if (PakIndexStart != INDEX_NONE)
+			{
+				PakIndexFromFilename = PakIndexFromFilename.RightChop(PakIndexStart + 1);
+				if (PakIndexFromFilename.IsNumeric())
+				{
+					PakPriority = FCString::Atoi(*PakIndexFromFilename);
+				}
+			}
+		}
+
+		FPakFile PakFile(*PakFilename, FParse::Param(FCommandLine::Get(), TEXT("signed")));
+		if (PakFile.IsValid())
+		{
+			FString DestPath(InDestPath);
+			FArchive& PakReader = *PakFile.GetSharedReader(NULL);
+			const int64 BufferSize = 8 * 1024 * 1024; // 8MB buffer for extracting
+			void* Buffer = FMemory::Malloc(BufferSize);
+			int64 CompressionBufferSize = 0;
+			uint8* PersistantCompressionBuffer = NULL;
+			int32 ErrorCount = 0;
+			int32 FileCount = 0;
+
+			FString PakMountPoint = bUseMountPoint ? PakFile.GetMountPoint().Replace(TEXT("../../../"), TEXT("")) : TEXT("");
+
+			for (FPakFile::FFileIterator It(PakFile); It; ++It, ++FileCount)
+			{
+				// Extract only the most recent version of a file when present in multiple paks
+				FFileInfo* HashFileInfo = InFileHashes.Find(It.Filename());
+				if (HashFileInfo == nullptr || HashFileInfo->PatchIndex == PakPriority)
+				{
+					const FPakEntry& Entry = It.Info();
+					PakReader.Seek(Entry.Offset);
+					uint32 SerializedCrcTest = 0;
+					FPakEntry EntryInfo;
+					EntryInfo.Serialize(PakReader, PakFile.GetInfo().Version);
+					if (EntryInfo == Entry)
+					{
+						FString DestFilename(DestPath / PakMountPoint / It.Filename());
+
+						TUniquePtr<FArchive> FileHandle(IFileManager::Get().CreateFileWriter(*DestFilename));
+						if (FileHandle)
+						{
+							if (Entry.CompressionMethod == COMPRESS_None)
+							{
+								BufferedCopyFile(*FileHandle, PakReader, Entry, Buffer, BufferSize, InEncryptionKey);
+							}
+							else
+							{
+								UncompressCopyFile(*FileHandle, PakReader, Entry, PersistantCompressionBuffer, CompressionBufferSize, InEncryptionKey, PakFile);
+							}
+							UE_LOG(LogPakFile, Display, TEXT("Extracted \"%s\" to \"%s\"."), *It.Filename(), *DestFilename);
+
+							if (OutOrderMap != nullptr)
+							{
+								OutOrderMap->Add(DestFilename, OutOrderMap->Num());
+							}
+
+							if (OutEntries != nullptr)
+							{
+								FPakInputPair Input;
+
+								Input.Source = DestFilename;
+								FPaths::NormalizeFilename(Input.Source);
+
+								Input.Dest = PakFile.GetMountPoint() + FPaths::GetPath(It.Filename());
+								FPaths::NormalizeFilename(Input.Dest);
+								FPakFile::MakeDirectoryFromPath(Input.Dest);
+
+								Input.bNeedsCompression = Entry.CompressionMethod != 0;
+								Input.bNeedEncryption = Entry.bEncrypted != 0;
+	
+								OutEntries->Add(Input);
+							}
+						}
+						else
+						{
+							UE_LOG(LogPakFile, Error, TEXT("Unable to create file \"%s\"."), *DestFilename);
+							ErrorCount++;
+						}
+					}
+					else
+					{
+						UE_LOG(LogPakFile, Error, TEXT("Serialized hash mismatch for \"%s\"."), *It.Filename());
+						ErrorCount++;
+					}
+				}
+			}
+			FMemory::Free(Buffer);
+			FMemory::Free(PersistantCompressionBuffer);
+
+			UE_LOG(LogPakFile, Log, TEXT("Finished extracting %d files (including %d errors)."), FileCount, ErrorCount);
+		}
+		else
+		{
+			UE_LOG(LogPakFile, Error, TEXT("Unable to open pak file \"%s\"."), *PakFilename);
+			return false;
+		}
+	}
+
+	return true;
+}
+
+void CreateDiffRelativePathMap(TArray<FString>& FileNames, const FString& RootPath, TMap<FName, FString>& OutMap)
+{
+	for (int32 i = 0; i < FileNames.Num(); ++i)
+	{
+		const FString& FullPath = FileNames[i];
+		FString RelativePath = FullPath.Mid(RootPath.Len());
+		OutMap.Add(FName(*RelativePath), FullPath);
+	}
+}
+
+bool DiffFilesInPaks(const FString InPakFilename1, const FString InPakFilename2, const FAES::FAESKey& InEncryptionKey)
+{
+	int32 NumUniquePAK1 = 0;
+	int32 NumUniquePAK2 = 0;
+	int32 NumDifferentContents = 0;
+	int32 NumEqualContents = 0;
+
+	TGuardValue<ELogTimes::Type> DisableLogTimes(GPrintLogTimes, ELogTimes::None);
+	UE_LOG(LogPakFile, Log, TEXT("FileEventType, FileName, Size1, Size2"));
+
+	// Allow the suppression of unique file logging for one or both files
+	const bool bLogUniques = !FParse::Param(FCommandLine::Get(), TEXT("nouniques"));
+	const bool bLogUniques1 = bLogUniques && !FParse::Param(FCommandLine::Get(), TEXT("nouniquesfile1"));
+	const bool bLogUniques2 = bLogUniques && !FParse::Param(FCommandLine::Get(), TEXT("nouniquesfile2"));
+
+	FPakFile PakFile1(*InPakFilename1, FParse::Param(FCommandLine::Get(), TEXT("signed")));
+	FPakFile PakFile2(*InPakFilename2, FParse::Param(FCommandLine::Get(), TEXT("signed")));
+	if (PakFile1.IsValid() && PakFile2.IsValid())
+	{		
+		FArchive& PakReader1 = *PakFile1.GetSharedReader(NULL);
+		FArchive& PakReader2 = *PakFile2.GetSharedReader(NULL);
+
+		const int64 BufferSize = 8 * 1024 * 1024; // 8MB buffer for extracting
+		void* Buffer = FMemory::Malloc(BufferSize);
+		int64 CompressionBufferSize = 0;
+		uint8* PersistantCompressionBuffer = NULL;
+		int32 ErrorCount = 0;
+		int32 FileCount = 0;		
+		
+		//loop over pak1 entries.  compare against entry in pak2.
+		for (FPakFile::FFileIterator It(PakFile1); It; ++It, ++FileCount)
+		{
+			const FString& PAK1FileName = It.Filename();
+
+			//double check entry info and move pakreader into place
+			const FPakEntry& Entry1 = It.Info();
+			PakReader1.Seek(Entry1.Offset);
+
+			FPakEntry EntryInfo1;
+			EntryInfo1.Serialize(PakReader1, PakFile1.GetInfo().Version);
+
+			if (EntryInfo1 != Entry1)
+			{
+				UE_LOG(LogPakFile, Log, TEXT("PakEntry1Invalid, %s, 0, 0"), *PAK1FileName);
+				continue;
+			}
+			
+			//see if entry exists in other pak							
+			FPakEntry Entry2;
+			bool bFoundEntry2 = PakFile2.Find(PakFile1.GetMountPoint() / PAK1FileName, &Entry2);
+			if (!bFoundEntry2)
+			{
+				++NumUniquePAK1;
+				if (bLogUniques1)
+				{
+					UE_LOG(LogPakFile, Log, TEXT("UniqueToFirstPak, %s, %i, 0"), *PAK1FileName, EntryInfo1.UncompressedSize);
+				}
+				continue;
+			}
+
+			//double check entry info and move pakreader into place
+			PakReader2.Seek(Entry2.Offset);
+			FPakEntry EntryInfo2;
+			EntryInfo2.Serialize(PakReader2, PakFile2.GetInfo().Version);
+			if (EntryInfo2 != Entry2)
+			{
+				UE_LOG(LogPakFile, Log, TEXT("PakEntry2Invalid, %s, 0, 0"), *PAK1FileName);
+				continue;;
+			}
+
+			//check sizes first as quick compare.
+			if (EntryInfo1.UncompressedSize != EntryInfo2.UncompressedSize)
+			{
+				UE_LOG(LogPakFile, Log, TEXT("FilesizeDifferent, %s, %i, %i"), *PAK1FileName, EntryInfo1.UncompressedSize, EntryInfo2.UncompressedSize);
+				continue;
+			}
+			
+			//serialize and memcompare the two entries
+			{
+				FLargeMemoryWriter PAKWriter1(EntryInfo1.UncompressedSize);
+				FLargeMemoryWriter PAKWriter2(EntryInfo2.UncompressedSize);
+
+				if (EntryInfo1.CompressionMethod == COMPRESS_None)
+				{
+					BufferedCopyFile(PAKWriter1, PakReader1, Entry1, Buffer, BufferSize, InEncryptionKey);
+				}
+				else
+				{
+					UncompressCopyFile(PAKWriter1, PakReader1, Entry1, PersistantCompressionBuffer, CompressionBufferSize, InEncryptionKey, PakFile1);
+				}
+
+				if (EntryInfo2.CompressionMethod == COMPRESS_None)
+				{
+					BufferedCopyFile(PAKWriter2, PakReader2, Entry2, Buffer, BufferSize, InEncryptionKey);
+				}
+				else
+				{
+					UncompressCopyFile(PAKWriter2, PakReader2, Entry2, PersistantCompressionBuffer, CompressionBufferSize, InEncryptionKey, PakFile2);
+				}
+
+				if (FMemory::Memcmp(PAKWriter1.GetData(), PAKWriter2.GetData(), EntryInfo1.UncompressedSize) != 0)
+				{
+					++NumDifferentContents;
+					UE_LOG(LogPakFile, Log, TEXT("ContentsDifferent, %s, %i, %i"), *PAK1FileName, EntryInfo1.UncompressedSize, EntryInfo2.UncompressedSize);
+				}
+				else
+				{
+					++NumEqualContents;
+				}
+			}			
+		}
+		
+		//check for files unique to the second pak.
+		for (FPakFile::FFileIterator It(PakFile2); It; ++It, ++FileCount)
+		{
+			const FPakEntry& Entry2 = It.Info();
+			PakReader2.Seek(Entry2.Offset);
+
+			FPakEntry EntryInfo2;
+			EntryInfo2.Serialize(PakReader2, PakFile2.GetInfo().Version);
+
+			if (EntryInfo2 == Entry2)
+			{
+				const FString& PAK2FileName = It.Filename();
+				FPakEntry Entry1;
+				bool bFoundEntry1 = PakFile1.Find(PakFile2.GetMountPoint() / PAK2FileName, &Entry1);
+				if (!bFoundEntry1)
+				{
+					++NumUniquePAK2;
+					if (bLogUniques2)
+					{
+						UE_LOG(LogPakFile, Log, TEXT("UniqueToSecondPak, %s, 0, %i"), *PAK2FileName, Entry2.UncompressedSize);
+					}
+					continue;
+				}
+			}
+		}
+
+		FMemory::Free(Buffer);
+		Buffer = nullptr;
+	}
+
+	UE_LOG(LogPakFile, Log, TEXT("Comparison complete"));
+	UE_LOG(LogPakFile, Log, TEXT("Unique to first pak: %i, Unique to second pak: %i, Num Different: %i, NumEqual: %i"), NumUniquePAK1, NumUniquePAK2, NumDifferentContents, NumEqualContents);	
+	return true;
+}
+
+void GenerateHashForFile(uint8* ByteBuffer, uint64 TotalSize, FFileInfo& FileHash)
+{
+	FMD5 FileHasher;
+	FileHasher.Update(ByteBuffer, TotalSize);
+	FileHasher.Final(FileHash.Hash);
+	FileHash.FileSize = TotalSize;
+}
+
+bool GenerateHashForFile( FString Filename, FFileInfo& FileHash)
+{
+	FArchive* File = IFileManager::Get().CreateFileReader(*Filename);
+
+	if ( File == NULL )
+		return false;
+
+	uint64 TotalSize = File->TotalSize();
+
+	uint8* ByteBuffer = new uint8[TotalSize];
+
+	File->Serialize(ByteBuffer, TotalSize);
+
+	delete File;
+	File = NULL;
+
+	GenerateHashForFile(ByteBuffer, TotalSize, FileHash);
+	
+	delete[] ByteBuffer;
+	return true;
+}
+
+bool GenerateHashesFromPak(const TCHAR* InPakFilename, const TCHAR* InDestPakFilename, TMap<FString, FFileInfo>& FileHashes, bool bUseMountPoint, const FAES::FAESKey& InEncryptionKey)
+{
+	TArray<FString> FoundFiles;
+	IFileManager::Get().FindFiles(FoundFiles, InPakFilename, true, false);
+	if (FoundFiles.Num() == 0)
+	{
+		return false;
+	}
+
+	// Gather all patch pak files and run through them one at a time
+	TArray<FString> PakFileList;
+	FString PakFileDirectory = FPaths::GetPath(InPakFilename);
+	IFileManager::Get().FindFiles(PakFileList, *PakFileDirectory, *FPaths::GetCleanFilename(InPakFilename));
+	for (int32 PakFileIndex = 0; PakFileIndex < PakFileList.Num(); PakFileIndex++)
+	{
+		FString PakFilename = PakFileDirectory + "\\" + PakFileList[PakFileIndex];
+		// Skip the destination pak file so we can regenerate an existing patch level
+		if (PakFilename.Equals(InDestPakFilename))
+		{
+			continue;
+		}
+		// Parse the pak file index, the base pak file is index -1
+		int32 PakPriority = -1;
+		if (PakFilename.EndsWith("_P.pak"))
+		{
+			FString PakIndexFromFilename = PakFilename.LeftChop(6);
+			int32 PakIndexStart = INDEX_NONE;
+			PakIndexFromFilename.FindLastChar('_', PakIndexStart);
+			if (PakIndexStart != INDEX_NONE)
+			{
+				PakIndexFromFilename = PakIndexFromFilename.RightChop(PakIndexStart + 1);
+				if (PakIndexFromFilename.IsNumeric())
+				{
+					PakPriority = FCString::Atoi(*PakIndexFromFilename);
+				}
+			}
+		}
+
+		FPakFile PakFile(*PakFilename, FParse::Param(FCommandLine::Get(), TEXT("signed")));
+		if (PakFile.IsValid())
+		{
+			FArchive& PakReader = *PakFile.GetSharedReader(NULL);
+			const int64 BufferSize = 8 * 1024 * 1024; // 8MB buffer for extracting
+			void* Buffer = FMemory::Malloc(BufferSize);
+			int64 CompressionBufferSize = 0;
+			uint8* PersistantCompressionBuffer = NULL;
+			int32 ErrorCount = 0;
+			int32 FileCount = 0;
+
+			FString PakMountPoint = bUseMountPoint ? PakFile.GetMountPoint().Replace(TEXT("../../../"), TEXT("")) : TEXT("");
+
+			for (FPakFile::FFileIterator It(PakFile); It; ++It, ++FileCount)
+			{
+				const FPakEntry& Entry = It.Info();
+				PakReader.Seek(Entry.Offset);
+				uint32 SerializedCrcTest = 0;
+				FPakEntry EntryInfo;
+				EntryInfo.Serialize(PakReader, PakFile.GetInfo().Version);
+				if (EntryInfo == Entry)
+				{
+					// TAutoPtr<FArchive> FileHandle(IFileManager::Get().CreateFileWriter(*DestFilename));
+					TArray<uint8> Bytes;
+					FMemoryWriter MemoryFile(Bytes);
+					FArchive* FileHandle = &MemoryFile;
+					// if (FileHandle.IsValid())
+					{
+						if (Entry.CompressionMethod == COMPRESS_None)
+						{
+							BufferedCopyFile(*FileHandle, PakReader, Entry, Buffer, BufferSize, InEncryptionKey);
+						}
+						else
+						{
+							UncompressCopyFile(*FileHandle, PakReader, Entry, PersistantCompressionBuffer, CompressionBufferSize, InEncryptionKey, PakFile);
+						}
+
+						FString FullFilename = PakMountPoint;
+						if (!FullFilename.IsEmpty() && !FullFilename.EndsWith("/"))
+						{
+							FullFilename += "/";
+						}
+						FullFilename += It.Filename();
+						UE_LOG(LogPakFile, Display, TEXT("Generated hash for \"%s\""), *FullFilename);
+						FFileInfo FileHash;
+						GenerateHashForFile(Bytes.GetData(), Bytes.Num(), FileHash);
+						FileHash.PatchIndex = PakPriority;
+
+						// Keep only the hash of the most recent version of a file (across multiple pak patch files)
+						if (!FileHashes.Contains(FullFilename))
+						{
+							FileHashes.Add(FullFilename, FileHash);
+						}
+						else if (FileHashes[FullFilename].PatchIndex < FileHash.PatchIndex)
+						{
+							FileHashes[FullFilename] = FileHash;
+						}
+					}
+					/*else
+					{
+					UE_LOG(LogPakFile, Error, TEXT("Unable to create file \"%s\"."), *DestFilename);
+					ErrorCount++;
+					}*/
+
+				}
+				else
+				{
+					UE_LOG(LogPakFile, Error, TEXT("Serialized hash mismatch for \"%s\"."), *It.Filename());
+					ErrorCount++;
+				}
+			}
+			FMemory::Free(Buffer);
+			FMemory::Free(PersistantCompressionBuffer);
+
+			UE_LOG(LogPakFile, Log, TEXT("Finished extracting %d files (including %d errors)."), FileCount, ErrorCount);
+		}
+		else
+		{
+			UE_LOG(LogPakFile, Error, TEXT("Unable to open pak file \"%s\"."), *PakFilename);
+			return false;
+		}
+	}
+
+	return true;
+}
+
+bool FileIsIdentical(FString SourceFile, FString DestFilename, const FFileInfo* Hash)
+{
+	int64 SourceTotalSize = Hash ? Hash->FileSize : IFileManager::Get().FileSize(*SourceFile);
+	int64 DestTotalSize = IFileManager::Get().FileSize(*DestFilename);
+
+	if (SourceTotalSize != DestTotalSize)
+	{
+		// file size doesn't match 
+		UE_LOG(LogPakFile, Display, TEXT("Source file size for %s %d bytes doesn't match %s %d bytes, did find %d"), *SourceFile, SourceTotalSize, *DestFilename, DestTotalSize, Hash ? 1 : 0);
+		return false;
+	}
+
+	FFileInfo SourceFileHash;
+	if (!Hash)
+	{
+		if (GenerateHashForFile(SourceFile, SourceFileHash) == false)
+		{
+			// file size doesn't match 
+			UE_LOG(LogPakFile, Display, TEXT("Source file size %s doesn't exist will be included in build"), *SourceFile);
+			return false;;
+		}
+		else
+		{
+			UE_LOG(LogPakFile, Warning, TEXT("Generated hash for file %s but it should have been in the FileHashes array"), *SourceFile);
+		}
+	}
+	else
+	{
+		SourceFileHash = *Hash;
+	}
+
+	FFileInfo DestFileHash;
+	if (GenerateHashForFile(DestFilename, DestFileHash) == false)
+	{
+		// destination file was removed don't really care about it
+		UE_LOG(LogPakFile, Display, TEXT("File was removed from destination cooked content %s not included in patch"), *DestFilename);
+		return false;
+	}
+
+	int32 Diff = FMemory::Memcmp(&SourceFileHash.Hash, &DestFileHash.Hash, sizeof(DestFileHash.Hash));
+	if (Diff != 0)
+	{
+		UE_LOG(LogPakFile, Display, TEXT("Source file hash for %s doesn't match dest file hash %s and will be included in patch"), *SourceFile, *DestFilename);
+		return false;
+	}
+
+	return true;
+}
+
+void RemoveIdenticalFiles( TArray<FPakInputPair>& FilesToPak, const FString& SourceDirectory, const TMap<FString, FFileInfo>& FileHashes )
+{
+	FString HashFilename = SourceDirectory / TEXT("Hashes.txt");
+
+	if (IFileManager::Get().FileExists(*HashFilename) )
+	{
+		FString EntireFile;
+		FFileHelper::LoadFileToString(EntireFile, *HashFilename);
+	}
+
+	TArray<FString> FilesToRemove;
+
+	for ( int I = FilesToPak.Num()-1; I >= 0; --I )
+	{
+		const auto& NewFile = FilesToPak[I]; 
+
+		FString SourceFileNoMountPoint =  NewFile.Dest.Replace(TEXT("../../../"), TEXT(""));
+		FString SourceFilename = SourceDirectory / NewFile.Dest.Replace(TEXT("../../../"), TEXT(""));
+		
+		const FFileInfo* FoundFileHash = FileHashes.Find(SourceFileNoMountPoint);
+		if (!FoundFileHash)
+		{
+			FoundFileHash = FileHashes.Find(NewFile.Dest);
+		}
+		
+ 		if ( !FoundFileHash )
+ 		{
+ 			UE_LOG(LogPakFile, Display, TEXT("Didn't find hash for %s No mount %s"), *SourceFilename, *SourceFileNoMountPoint);
+ 		}
+ 
+		// uexp files are always handled with their corresponding uasset file
+		if (!FPaths::GetExtension(SourceFilename).Equals("uexp", ESearchCase::IgnoreCase))
+		{
+			FString DestFilename = NewFile.Source;
+			if (FileIsIdentical(SourceFilename, DestFilename, FoundFileHash))
+			{
+				// Check for uexp files only for uasset files
+				FString Ext(FPaths::GetExtension(SourceFilename));
+				if (Ext.Equals("uasset", ESearchCase::IgnoreCase) || Ext.Equals("umap", ESearchCase::IgnoreCase))
+				{
+					FString UexpSourceFilename = FPaths::ChangeExtension(SourceFilename, "uexp");
+					FString UexpSourceFileNoMountPoint = FPaths::ChangeExtension(SourceFileNoMountPoint, "uexp");
+
+					const FFileInfo* UexpFoundFileHash = FileHashes.Find(UexpSourceFileNoMountPoint);
+					if (!UexpFoundFileHash)
+					{
+						UexpFoundFileHash = FileHashes.Find(FPaths::ChangeExtension(NewFile.Dest, "uexp"));
+					}
+
+					if (!UexpFoundFileHash)
+					{
+						UE_LOG(LogPakFile, Display, TEXT("Didn't find hash for %s No mount %s"), *UexpSourceFilename, *UexpSourceFileNoMountPoint);
+					}
+
+					if (UexpFoundFileHash || IFileManager::Get().FileExists(*UexpSourceFilename))
+					{
+
+						FString UexpDestFilename = FPaths::ChangeExtension(NewFile.Source, "uexp");
+						if (!FileIsIdentical(UexpSourceFilename, UexpDestFilename, UexpFoundFileHash))
+						{
+							UE_LOG(LogPakFile, Display, TEXT("%s not identical for %s. Including both files in patch."), *UexpSourceFilename, *SourceFilename);
+							continue;
+						}
+						// Add this file to the list to be removed from FilesToPak after we finish processing (since this file was found at random within 
+						// the list we cannot remove it or we'll mess up our containing for loop)
+						FilesToRemove.Add(UexpDestFilename);
+					}
+				}
+
+				UE_LOG(LogPakFile, Display, TEXT("Source file %s matches dest file %s and will not be included in patch"), *SourceFilename, *DestFilename);
+				// remove from the files to pak list
+				FilesToPak.RemoveAt(I);
+			}
+		}
+	}
+
+	// Clean up uexp files that were marked for removal, assume files may only be listed one in FilesToPak
+	for (int FileIndexToRemove = 0; FileIndexToRemove < FilesToRemove.Num(); FileIndexToRemove++)
+	{
+		const FPakInputPair FileSourceToRemove(FilesToRemove[FileIndexToRemove], "");
+		FilesToPak.RemoveSingle(FileSourceToRemove);
+	}
+}
+
+FString GetPakPath(const TCHAR* SpecifiedPath, bool bIsForCreation)
+{
+	FString PakFilename(SpecifiedPath);
+	FPaths::MakeStandardFilename(PakFilename);
+	
+	// if we are trying to open (not create) it, but BaseDir relative doesn't exist, look in LaunchDir
+	if (!bIsForCreation && !FPaths::FileExists(PakFilename))
+	{
+		PakFilename = FPaths::LaunchDir() + SpecifiedPath;
+
+		if (!FPaths::FileExists(PakFilename))
+		{
+			UE_LOG(LogPakFile, Fatal, TEXT("Existing pak file %s could not be found (checked against binary and launch directories)"), SpecifiedPath);
+			return TEXT("");
+		}
+	}
+	
+	return PakFilename;
+}
+
+struct FPackage
+{
+public:
+	TSet<FPackage*> DirectlyReferencing;
+	TSet<FPackage*> DirectlyReferencedBy;
+	TSet<FPackage*> AllReferences;
+
+	FName Name;
+	int64 InclusiveSize;
+	int64 ExclusiveSize;
+	int Id;
+
+private:
+	bool bUpdateHelper;
+	static TMap<FName,FPackage*> NameToPackageMap;
+
+public:
+	FPackage(const FName& InName,int InId)
+		: Name(InName)
+		, InclusiveSize(0)
+		, ExclusiveSize(0)
+		, Id(InId)
+		, bUpdateHelper(false)
+	{}
+
+	static FPackage* FindOrCreatePackage( FName PackageName )
+	{
+		static int Id = 1;
+		FPackage* Package = NameToPackageMap.FindRef(PackageName);
+		if(!Package)
+		{
+			Package = new FPackage(PackageName,Id++);
+			NameToPackageMap.Add(PackageName,Package);
+		}
+		return Package;
+	}
+
+	void ResetUpdateHelper()
+	{
+		bUpdateHelper = false;
+	}
+
+	void RecurseUpdateReferences()
+	{
+		if( !bUpdateHelper )
+		{
+			bUpdateHelper = true;
+			for( auto& DirectReference : DirectlyReferencing )
+			{
+				AllReferences.Add(DirectReference);
+				DirectReference->RecurseUpdateReferences();
+				AllReferences.Append(DirectReference->AllReferences);
+			}
+		}
+	}
+
+	void UpdateInclusiveSize()
+	{
+		InclusiveSize = ExclusiveSize;
+		for(auto& Reference : AllReferences)
+		{
+			InclusiveSize += Reference->ExclusiveSize;
+		}
+	}
+
+	static void GetAllPackages( TArray<FPackage*>& OutPackages )
+	{
+		OutPackages.Reset(NameToPackageMap.Num());
+		for( const auto& Entry : NameToPackageMap )
+		{
+			OutPackages.Add(Entry.Value);
+		}
+	}
+
+	TArray< TSharedPtr<FJsonValue> > ToJsonHelper( const TSet<FPackage*>& Packages )
+	{
+		TArray< TSharedPtr<FJsonValue> > JsonPackageNames;
+		for( const auto Package : Packages )
+		{
+			JsonPackageNames.Add(MakeShareable(new FJsonValueString(Package->Name.ToString())));
+		}
+		return JsonPackageNames;
+	}
+
+	TSharedPtr<FJsonObject> ToJsonObject()
+	{
+		TSharedPtr<FJsonObject> JsonPackageObject = MakeShareable(new FJsonObject);
+		
+		JsonPackageObject->SetStringField(TEXT("Name"),*Name.ToString());
+		JsonPackageObject->SetNumberField(TEXT("InclusiveSize"),InclusiveSize);
+		JsonPackageObject->SetNumberField(TEXT("ExclusiveSize"),ExclusiveSize);
+
+		JsonPackageObject->SetArrayField(TEXT("DirectlyReferencing"),ToJsonHelper(DirectlyReferencing));
+		JsonPackageObject->SetArrayField(TEXT("DirectlyReferencedBy"),ToJsonHelper(DirectlyReferencedBy));
+		JsonPackageObject->SetArrayField(TEXT("AllReferences"),ToJsonHelper(AllReferences));
+
+		return JsonPackageObject;
+	}
+};
+TMap<FName,FPackage*> FPackage::NameToPackageMap;
+
+bool ExportDependencies(const TCHAR * PakFilename, const TCHAR* GameName, const TCHAR* GameFolderName, const TCHAR* OutputFilenameBase )
+{
+	// Example command line used for this tool
+	// C:\Development\BB\WEX\Saved\StagedBuilds\WindowsNoEditor\WorldExplorers\Content\Paks\WorldExplorers-WindowsNoEditor.pak WorldExplorers WEX -exportdependencies=c:\dvtemp\output -debug -NoAssetRegistryCache -ForceDependsGathering
+	
+	FPakFile PakFile(PakFilename,FParse::Param(FCommandLine::Get(),TEXT("signed")));
+
+	if(PakFile.IsValid())
+	{
+		// Get size information from PAK file.
+		{
+			TArray<FPakFile::FFileIterator> Records;
+			FString PakGameContentFolder = FString(GameName) + TEXT("/Content");
+			for(FPakFile::FFileIterator It(PakFile); It; ++It)
+			{
+				FString PackageName;
+				It.Filename().Split(TEXT("."),&PackageName,NULL);
+				int64 Size = It.Info().Size;			
+
+				if( PackageName.StartsWith(TEXT("Engine/Content")) )
+				{
+					PackageName = PackageName.Replace(TEXT("Engine/Content"),TEXT("/Engine"));
+				}
+				else if( PackageName.StartsWith(*PakGameContentFolder))
+				{
+					PackageName = PackageName.Replace(*PakGameContentFolder,TEXT("/Game"));
+				}
+
+				FPackage* Package = FPackage::FindOrCreatePackage(FName(*PackageName));
+				Package->ExclusiveSize += Size;
+			}
+		}
+
+		TMap<FName,FName> PackageToClassMap;
+
+		// Combine with dependency information from asset registry.
+		{
+			FString ProjectDir = FString(TEXT("../../../")) + GameFolderName + TEXT("/");
+			FPlatformMisc::SetOverrideProjectDir(*ProjectDir);
+			FAssetRegistryModule& AssetRegistryModule = FModuleManager::Get().LoadModuleChecked<FAssetRegistryModule>(TEXT("AssetRegistry"));
+		
+			AssetRegistryModule.Get().SearchAllAssets(true);
+			TArray<FAssetData> AssetData;		
+			AssetRegistryModule.Get().GetAllAssets(AssetData,true);	
+	
+			TSet<FName> PackageNames;
+			for( int i=0; i<AssetData.Num(); i++ )
+			{
+				PackageNames.Add(AssetData[i].PackageName);
+				PackageToClassMap.Add(AssetData[i].PackageName,AssetData[i].AssetClass);
+			}
+
+			for( const auto& PackageName : PackageNames )
+			{
+				TArray<FName> DependencyArray;
+				AssetRegistryModule.Get().GetDependencies(PackageName,DependencyArray);
+
+				FPackage* Package = FPackage::FindOrCreatePackage(PackageName);
+				for( const auto& DependencyName : DependencyArray )
+				{
+					// exclude '/Script/' as it clutters up things significantly.
+					if( !DependencyName.ToString().StartsWith(TEXT("/Script/")) )
+					{
+						FPackage* Dependency = FPackage::FindOrCreatePackage(DependencyName);
+						Package->DirectlyReferencing.Add(Dependency);
+						Dependency->DirectlyReferencedBy.Add(Package);
+					}
+				}
+			}
+
+			// 2 passes are required to deal with cycles.
+			for(const auto& PackageName : PackageNames)
+			{
+				FPackage* Package = FPackage::FindOrCreatePackage(PackageName);
+				Package->RecurseUpdateReferences();
+			}
+			for(const auto& PackageName : PackageNames)
+			{
+				FPackage* Package = FPackage::FindOrCreatePackage(PackageName);
+				Package->ResetUpdateHelper();
+			}
+			for(const auto& PackageName : PackageNames)
+			{
+				FPackage* Package = FPackage::FindOrCreatePackage(PackageName);
+				Package->RecurseUpdateReferences();
+			}
+		}
+
+		// Update inclusive size, asset class, and export to CSV, JSON, and GDF
+		{
+			TSharedPtr<FJsonObject> JsonRootObject = MakeShareable(new FJsonObject);
+			TArray< TSharedPtr<FJsonValue> > JsonPackages;
+
+			TArray<FPackage*> AllPackages;
+			FPackage::GetAllPackages(AllPackages);
+
+			for(auto Package : AllPackages)
+			{
+				Package->UpdateInclusiveSize();
+				JsonPackages.Add( MakeShareable(new FJsonValueObject(Package->ToJsonObject())) );
+			}
+			JsonRootObject->SetArrayField(TEXT("Packages"),JsonPackages);
+
+			FString JsonOutputString;
+			TSharedRef<TJsonWriter<TCHAR,TPrettyJsonPrintPolicy<TCHAR>>> JsonWriter = TJsonWriterFactory<TCHAR,TPrettyJsonPrintPolicy<TCHAR>>::Create(&JsonOutputString);
+			FJsonSerializer::Serialize(JsonRootObject.ToSharedRef(),JsonWriter);
+
+			FArchive* JsonFileWriter = IFileManager::Get().CreateFileWriter(*(FString(OutputFilenameBase)+TEXT(".json")));
+			if(JsonFileWriter)
+			{
+				JsonFileWriter->Logf(TEXT("%s"),*JsonOutputString);
+				JsonFileWriter->Close();
+				delete JsonFileWriter;
+			}
+
+			FArchive* CSVFileWriter = IFileManager::Get().CreateFileWriter(*(FString(OutputFilenameBase)+TEXT(".csv")));
+			if(CSVFileWriter)
+			{
+				CSVFileWriter->Logf(TEXT("class,name,inclusive,exclusive"));
+				for(auto Package : AllPackages)
+				{
+					FName ClassName = PackageToClassMap.FindRef(Package->Name);
+					CSVFileWriter->Logf(TEXT("%s,%s,%i,%i"),*ClassName.ToString(),*Package->Name.ToString(),Package->InclusiveSize,Package->ExclusiveSize);
+				}
+				CSVFileWriter->Close();
+				delete CSVFileWriter;
+				CSVFileWriter = NULL;
+			}
+
+			FArchive* GDFFileWriter = IFileManager::Get().CreateFileWriter(*(FString(OutputFilenameBase)+TEXT(".gdf")));
+			if(GDFFileWriter)
+			{
+				GDFFileWriter->Logf(TEXT("nodedef> name VARCHAR,label VARCHAR,inclusive DOUBLE,exclusive DOUBLE"));
+				GDFFileWriter->Logf(TEXT("0,root,0,0"));
+				for(auto Package : AllPackages)
+				{
+					GDFFileWriter->Logf(TEXT("%i,%s,%i,%i"),Package->Id,*Package->Name.ToString(),Package->InclusiveSize,Package->ExclusiveSize);
+				}
+				GDFFileWriter->Logf(TEXT("edgedef> node1 VARCHAR,node2 VARCHAR"));
+				// fake root to ensure spanning tree
+				for(auto Package : AllPackages)
+				{
+					GDFFileWriter->Logf(TEXT("0,%i"),Package->Id);
+				}
+				for(auto Package : AllPackages)
+				{
+					for( auto ReferencedPackage : Package->DirectlyReferencing )
+					{
+						GDFFileWriter->Logf(TEXT("%i,%i"),Package->Id,ReferencedPackage->Id);
+					}
+				}
+				GDFFileWriter->Close();
+				delete GDFFileWriter;
+				GDFFileWriter = NULL;
+			}
+		}
+		return true;
+	}
+	else
+	{
+		return false;
+	}
+}
+
+bool Repack(const FString& InputPakFile, const FString& OutputPakFile, const FPakCommandLineParameters& CmdLineParameters, const FKeyPair& SigningKey, const FAES::FAESKey& InEncryptionKey)
+{
+	bool bResult = false;
+
+	// Extract the existing pak file
+	TMap<FString, FFileInfo> Hashes;
+	TArray<FPakInputPair> Entries;
+	TMap<FString, uint64> OrderMap;
+	FString TempDir = FPaths::EngineIntermediateDir() / TEXT("UnrealPak") / TEXT("Repack") / FPaths::GetBaseFilename(InputPakFile);
+	if (ExtractFilesFromPak(*InputPakFile, Hashes, *TempDir, false, InEncryptionKey, &Entries, &OrderMap))
+	{
+		TArray<FPakInputPair> FilesToAdd;
+		CollectFilesToAdd(FilesToAdd, Entries, OrderMap);
+
+		// Get a temporary output filename. We'll only create/replace the final output file once successful.
+		FString TempOutputPakFile = FPaths::CreateTempFilename(*FPaths::GetPath(OutputPakFile), *FPaths::GetCleanFilename(OutputPakFile));
+
+		// Create the new pak file
+		UE_LOG(LogPakFile, Display, TEXT("Creating %s..."), *OutputPakFile);
+		if (CreatePakFile(*TempOutputPakFile, FilesToAdd, CmdLineParameters, SigningKey, InEncryptionKey))
+		{
+			IFileManager::Get().Move(*OutputPakFile, *TempOutputPakFile);
+
+			FString OutputSigFile = FPaths::ChangeExtension(OutputPakFile, TEXT(".sig"));
+			if (IFileManager::Get().FileExists(*OutputSigFile))
+			{
+				IFileManager::Get().Delete(*OutputSigFile);
+			}
+
+			FString TempOutputSigFile = FPaths::ChangeExtension(TempOutputPakFile, TEXT(".sig"));
+			if (IFileManager::Get().FileExists(*TempOutputSigFile))
+			{
+				IFileManager::Get().Move(*OutputSigFile, *TempOutputSigFile);
+			}
+
+			bResult = true;
+		}
+	}
+	IFileManager::Get().DeleteDirectory(*TempDir, false, true);
+
+	return bResult;
+}
+
+/**
+ * Application entry point
+ * Params:
+ *   -Test test if the pak file is healthy
+ *   -Extract extracts pak file contents (followed by a path, i.e.: -extract D:\ExtractedPak)
+ *   -Create=filename response file to create a pak file with
+ *   -Sign=filename use the key pair in filename to sign a pak file, or: -sign=key_hex_values_separated_with_+, i.e: -sign=0x123456789abcdef+0x1234567+0x12345abc
+ *    where the first number is the private key exponend, the second one is modulus and the third one is the public key exponent.
+ *   -Signed use with -extract and -test to let the code know this is a signed pak
+ *   -GenerateKeys=filename generates encryption key pair for signing a pak file
+ *   -P=prime will use a predefined prime number for generating encryption key file
+ *   -Q=prime same as above, P != Q, GCD(P, Q) = 1 (which is always true if they're both prime)
+ *   -GeneratePrimeTable=filename generates a prime table for faster prime number generation (.inl file)
+ *   -TableMax=number maximum prime number in the generated table (default is 10000)
+ *
+ * @param	ArgC	Command-line argument count
+ * @param	ArgV	Argument strings
+ */
+int32 ExecuteUnrealPak(int32 ArgC, TCHAR* ArgV[])
+{
+	if (ArgC < 2)
+	{
+		UE_LOG(LogPakFile, Error, TEXT("No pak file name specified. Usage:"));
+		UE_LOG(LogPakFile, Error, TEXT("  UnrealPak <PakFilename> -Test"));
+		UE_LOG(LogPakFile, Error, TEXT("  UnrealPak <PakFilename> -List"));
+		UE_LOG(LogPakFile, Error, TEXT("  UnrealPak <PakFilename> <GameUProjectName> <GameFolderName> -ExportDependencies=<OutputFileBase> -NoAssetRegistryCache -ForceDependsGathering"));
+		UE_LOG(LogPakFile, Error, TEXT("  UnrealPak <PakFilename> -Extract <ExtractDir>"));
+		UE_LOG(LogPakFile, Error, TEXT("  UnrealPak <PakFilename> -Create=<ResponseFile> [Options]"));
+		UE_LOG(LogPakFile, Error, TEXT("  UnrealPak <PakFilename> -Dest=<MountPoint>"));
+		UE_LOG(LogPakFile, Error, TEXT("  UnrealPak <PakFilename> -Repack [-Output=Path] [Options]"));
+		UE_LOG(LogPakFile, Error, TEXT("  UnrealPak GenerateKeys=<KeyFilename>"));
+		UE_LOG(LogPakFile, Error, TEXT("  UnrealPak GeneratePrimeTable=<KeyFilename> [-TableMax=<N>]"));
+		UE_LOG(LogPakFile, Error, TEXT("  UnrealPak <PakFilename1> <PakFilename2> -diff"));
+		UE_LOG(LogPakFile, Error, TEXT("  UnrealPak -TestEncryption"));
+		UE_LOG(LogPakFile, Error, TEXT("  Options:"));
+		UE_LOG(LogPakFile, Error, TEXT("    -blocksize=<BlockSize>"));
+		UE_LOG(LogPakFile, Error, TEXT("    -bitwindow=<BitWindow>"));
+		UE_LOG(LogPakFile, Error, TEXT("    -compress"));
+		UE_LOG(LogPakFile, Error, TEXT("    -encrypt"));
+		UE_LOG(LogPakFile, Error, TEXT("    -order=<OrderingFile>"));
+		UE_LOG(LogPakFile, Error, TEXT("    -diff (requires 2 filenames first)"));
+		UE_LOG(LogPakFile, Error, TEXT("    -enginedir (specify engine dir for when using ini encryption configs)"));
+		UE_LOG(LogPakFile, Error, TEXT("    -projectdir (specify project dir for when using ini encryption configs)"));
+		UE_LOG(LogPakFile, Error, TEXT("    -encryptionini (specify ini base name to gather encryption settings from)"));
+		UE_LOG(LogPakFile, Error, TEXT("    -extracttomountpoint (Extract to mount point path of pak file)"));
+		UE_LOG(LogPakFile, Error, TEXT("    -encryptindex (encrypt the pak file index, making it unusable in unrealpak without supplying the key)"));
+		UE_LOG(LogPakFile, Error, TEXT("    -compressor=<DllPath> (register a custom compressor)"))
+		UE_LOG(LogPakFile, Error, TEXT("    -overrideplatformcompressor (override the native platform compressor)"))
+		return 1;
+	}
+
+	double StartTime = FPlatformTime::Seconds();
+	FKeyPair SigningKey;
+	FAES::FAESKey EncryptionKey;
+	PrepareEncryptionAndSigningKeys(SigningKey, EncryptionKey);
+
+	FPakCommandLineParameters CmdLineParameters;
+	int32 Result = 0;
+	FString KeyFilename;
+	if (FParse::Value(FCommandLine::Get(), TEXT("GenerateKeys="), KeyFilename, false))
+	{
+		Result = GenerateKeys(*KeyFilename) ? 0 : 1;
+	}
+	else if (FParse::Value(FCommandLine::Get(), TEXT("GeneratePrimeTable="), KeyFilename, false))
+	{
+		int64 MaxPrimeValue = 10000;
+		FParse::Value(FCommandLine::Get(), TEXT("TableMax="), MaxPrimeValue);
+		GeneratePrimeNumberTable(MaxPrimeValue, *KeyFilename);
+	}
+	else if (FParse::Param(FCommandLine::Get(), TEXT("TestEncryption")))
+	{
+		void TestEncryption();
+		TestEncryption();
+	}
+	else 
+	{
+		FString ExportDependencyFilename;
+		if (FParse::Param(FCommandLine::Get(), TEXT("Test")))
+		{
+			FString PakFilename = GetPakPath(ArgV[1], false);
+			Result = TestPakFile(*PakFilename) ? 0 : 1;
+		}
+		else if (FParse::Param(FCommandLine::Get(), TEXT("List")))
+		{
+			int64 SizeFilter = 0;
+			FParse::Value(FCommandLine::Get(), TEXT("SizeFilter="), SizeFilter);
+
+			FString PakFilename = GetPakPath(ArgV[1], false);
+			Result = ListFilesInPak(*PakFilename, SizeFilter) ? 0 : 1;
+		}
+		else if(FParse::Value(FCommandLine::Get(),TEXT("ExportDependencies="),ExportDependencyFilename,false))
+		{			
+			if(ArgC < 4)
+			{
+				UE_LOG(LogPakFile,Error,TEXT("Insufficient arguments."));
+			}
+			else
+			{
+				FString PakFilename = GetPakPath(ArgV[1],false);
+				Result = ExportDependencies(*PakFilename,ArgV[2],ArgV[3],*ExportDependencyFilename);					
+			}
+		}
+		else if (FParse::Param(FCommandLine::Get(), TEXT("Diff")))
+		{
+			FString PakFilename1 = GetPakPath(ArgV[1], false);
+			FString PakFilename2 = GetPakPath(ArgV[2], false);
+			Result = DiffFilesInPaks(*PakFilename1, *PakFilename2, EncryptionKey) ? 0 : 1;
+		}
+		else if (FParse::Param(FCommandLine::Get(), TEXT("Extract")))
+		{
+			FString PakFilename = GetPakPath(ArgV[1], false);
+			if (ArgC < 4)
+			{
+				UE_LOG(LogPakFile, Error, TEXT("No extraction path specified."));
+				Result = 1;
+			}
+			else
+			{
+				FString DestPath = (ArgV[2][0] == '-') ? ArgV[3] : ArgV[2];
+				bool bExtractToMountPoint = FParse::Param(FCommandLine::Get(), TEXT("ExtractToMountPoint"));
+				TMap<FString, FFileInfo> EmptyMap;
+				Result = ExtractFilesFromPak(*PakFilename, EmptyMap, *DestPath, bExtractToMountPoint, EncryptionKey) ? 0 : 1;
+			}
+		}
+		else if (FParse::Param(FCommandLine::Get(), TEXT("Repack")))
+		{
+			TArray<FPakInputPair> Entries;
+			ProcessCommandLine(ArgC, ArgV, Entries, CmdLineParameters);
+
+			// Find all the input pak files
+			FString InputDir = FPaths::GetPath(ArgV[1]);
+
+			TArray<FString> InputPakFiles;
+			IFileManager::Get().FindFiles(InputPakFiles, *InputDir, *FPaths::GetCleanFilename(ArgV[1]));
+
+			for (int Idx = 0; Idx < InputPakFiles.Num(); Idx++)
+			{
+				InputPakFiles[Idx] = InputDir / InputPakFiles[Idx];
+			}
+
+			if (InputPakFiles.Num() == 0)
+			{
+				UE_LOG(LogPakFile, Error, TEXT("No files found matching '%s'"), ArgV[1]);
+				Result = 1;
+			}
+			else
+			{
+				// Find all the output paths
+				TArray<FString> OutputPakFiles;
+
+				FString OutputPath;
+				if (!FParse::Value(FCommandLine::Get(), TEXT("Output="), OutputPath, false))
+				{
+					for (const FString& InputPakFile : InputPakFiles)
+					{
+						OutputPakFiles.Add(InputPakFile);
+					}
+				}
+				else if (IFileManager::Get().DirectoryExists(*OutputPath))
+				{
+					for (const FString& InputPakFile : InputPakFiles)
+					{
+						OutputPakFiles.Add(FPaths::Combine(OutputPath, FPaths::GetCleanFilename(InputPakFile)));
+					}
+				}
+				else
+				{
+					for (const FString& InputPakFile : InputPakFiles)
+					{
+						OutputPakFiles.Add(OutputPath);
+					}
+				}
+
+				// Repack them all
+				for (int Idx = 0; Idx < InputPakFiles.Num(); Idx++)
+				{
+					UE_LOG(LogPakFile, Display, TEXT("Repacking %s into %s"), *InputPakFiles[Idx], *OutputPakFiles[Idx]);
+					if (!Repack(InputPakFiles[Idx], OutputPakFiles[Idx], CmdLineParameters, SigningKey, EncryptionKey))
+					{
+						Result = 1;
+						break;
+					}
+				}
+			}
+		}
+		else
+		{
+			// since this is for creation, we pass true to make it not look in LaunchDir
+			FString PakFilename = GetPakPath(ArgV[1], true);
+
+			// List of all items to add to pak file
+			TArray<FPakInputPair> Entries;
+			ProcessCommandLine(ArgC, ArgV, Entries, CmdLineParameters);
+			TMap<FString, uint64> OrderMap;
+			ProcessOrderFile(ArgC, ArgV, OrderMap);
+
+			if (Entries.Num() == 0)
+			{
+				UE_LOG(LogPakFile, Error, TEXT("No files specified to add to pak file."));
+				Result = 1;
+			}
+			else
+			{
+				TMap<FString, FFileInfo> SourceFileHashes;
+
+				if ( CmdLineParameters.GeneratePatch )
+				{
+					FString OutputPath;
+					if (!FParse::Value(FCommandLine::Get(), TEXT("TempFiles="), OutputPath))
+					{
+						OutputPath = FPaths::GetPath(PakFilename) / FString(TEXT("TempFiles"));
+					}
+
+					IFileManager::Get().DeleteDirectory(*OutputPath);
+
+					// Check command line for the "patchcryptokeys" param, which will tell us where to look for the encryption keys that
+					// we need to access the patch reference data
+					FString PatchReferenceCryptoKeysFilename;
+					FAES::FAESKey PatchReferenceEncryptionKey = EncryptionKey;
+					if (FParse::Value(FCommandLine::Get(), TEXT("PatchCryptoKeys="), PatchReferenceCryptoKeysFilename))
+					{
+						FKeyPair UnusedSigningKey;
+						PrepareEncryptionAndSigningKeysFromCryptoKeyCache(PatchReferenceCryptoKeysFilename, UnusedSigningKey, PatchReferenceEncryptionKey);
+					}
+
+					UE_LOG(LogPakFile, Display, TEXT("Generating patch from %s."), *CmdLineParameters.SourcePatchPakFilename, true );
+
+					if (!GenerateHashesFromPak(*CmdLineParameters.SourcePatchPakFilename, *PakFilename, SourceFileHashes, true, PatchReferenceEncryptionKey))
+					{
+						if (ExtractFilesFromPak(*CmdLineParameters.SourcePatchPakFilename, SourceFileHashes, *OutputPath, true, PatchReferenceEncryptionKey) == false)
+						{
+							UE_LOG(LogPakFile, Warning, TEXT("Unable to extract files from source pak file for patch"));
+						}
+						else
+						{
+							CmdLineParameters.SourcePatchDiffDirectory = OutputPath;
+						}
+					}
+				}
+
+
+				// Start collecting files
+				TArray<FPakInputPair> FilesToAdd;
+				CollectFilesToAdd(FilesToAdd, Entries, OrderMap);
+
+				if ( CmdLineParameters.GeneratePatch )
+				{
+					// if we are generating a patch here we remove files which are already shipped...
+					RemoveIdenticalFiles(FilesToAdd, CmdLineParameters.SourcePatchDiffDirectory, SourceFileHashes);
+				}
+
+
+				Result = CreatePakFile(*PakFilename, FilesToAdd, CmdLineParameters, SigningKey, EncryptionKey) ? 0 : 1;
+
+				if (CmdLineParameters.GeneratePatch)
+				{
+					FString OutputPath = FPaths::GetPath(PakFilename) / FString(TEXT("TempFiles"));
+					// delete the temporary directory
+					IFileManager::Get().DeleteDirectory(*OutputPath, false, true);
+				}
+			}
+		}
+	}
+
+	UE_LOG(LogPakFile, Display, TEXT("Unreal pak executed in %f seconds"), FPlatformTime::Seconds() - StartTime );
+
+	return Result;
+}
diff --git a/Engine/Source/Developer/PakFileUtilities/Private/Primes.inl b/Engine/Source/Developer/PakFileUtilities/Private/Primes.inl
new file mode 100644
--- /dev/null
+++ b/Engine/Source/Developer/PakFileUtilities/Private/Primes.inl
@@ -0,0 +1,127 @@
+// Copyright 1998-2018 Epic Games, Inc. All Rights Reserved.
+int32 PrimeTable[] =
+{
+	2, 3, 5, 7, 11, 13, 17, 19, 23, 29, 
+	31, 37, 41, 43, 47, 53, 59, 61, 67, 71, 
+	73, 79, 83, 89, 97, 101, 103, 107, 109, 113, 
+	127, 131, 137, 139, 149, 151, 157, 163, 167, 173, 
+	179, 181, 191, 193, 197, 199, 211, 223, 227, 229, 
+	233, 239, 241, 251, 257, 263, 269, 271, 277, 281, 
+	283, 293, 307, 311, 313, 317, 331, 337, 347, 349, 
+	353, 359, 367, 373, 379, 383, 389, 397, 401, 409, 
+	419, 421, 431, 433, 439, 443, 449, 457, 461, 463, 
+	467, 479, 487, 491, 499, 503, 509, 521, 523, 541, 
+	547, 557, 563, 569, 571, 577, 587, 593, 599, 601, 
+	607, 613, 617, 619, 631, 641, 643, 647, 653, 659, 
+	661, 673, 677, 683, 691, 701, 709, 719, 727, 733, 
+	739, 743, 751, 757, 761, 769, 773, 787, 797, 809, 
+	811, 821, 823, 827, 829, 839, 853, 857, 859, 863, 
+	877, 881, 883, 887, 907, 911, 919, 929, 937, 941, 
+	947, 953, 967, 971, 977, 983, 991, 997, 1009, 1013, 
+	1019, 1021, 1031, 1033, 1039, 1049, 1051, 1061, 1063, 1069, 
+	1087, 1091, 1093, 1097, 1103, 1109, 1117, 1123, 1129, 1151, 
+	1153, 1163, 1171, 1181, 1187, 1193, 1201, 1213, 1217, 1223, 
+	1229, 1231, 1237, 1249, 1259, 1277, 1279, 1283, 1289, 1291, 
+	1297, 1301, 1303, 1307, 1319, 1321, 1327, 1361, 1367, 1373, 
+	1381, 1399, 1409, 1423, 1427, 1429, 1433, 1439, 1447, 1451, 
+	1453, 1459, 1471, 1481, 1483, 1487, 1489, 1493, 1499, 1511, 
+	1523, 1531, 1543, 1549, 1553, 1559, 1567, 1571, 1579, 1583, 
+	1597, 1601, 1607, 1609, 1613, 1619, 1621, 1627, 1637, 1657, 
+	1663, 1667, 1669, 1693, 1697, 1699, 1709, 1721, 1723, 1733, 
+	1741, 1747, 1753, 1759, 1777, 1783, 1787, 1789, 1801, 1811, 
+	1823, 1831, 1847, 1861, 1867, 1871, 1873, 1877, 1879, 1889, 
+	1901, 1907, 1913, 1931, 1933, 1949, 1951, 1973, 1979, 1987, 
+	1993, 1997, 1999, 2003, 2011, 2017, 2027, 2029, 2039, 2053, 
+	2063, 2069, 2081, 2083, 2087, 2089, 2099, 2111, 2113, 2129, 
+	2131, 2137, 2141, 2143, 2153, 2161, 2179, 2203, 2207, 2213, 
+	2221, 2237, 2239, 2243, 2251, 2267, 2269, 2273, 2281, 2287, 
+	2293, 2297, 2309, 2311, 2333, 2339, 2341, 2347, 2351, 2357, 
+	2371, 2377, 2381, 2383, 2389, 2393, 2399, 2411, 2417, 2423, 
+	2437, 2441, 2447, 2459, 2467, 2473, 2477, 2503, 2521, 2531, 
+	2539, 2543, 2549, 2551, 2557, 2579, 2591, 2593, 2609, 2617, 
+	2621, 2633, 2647, 2657, 2659, 2663, 2671, 2677, 2683, 2687, 
+	2689, 2693, 2699, 2707, 2711, 2713, 2719, 2729, 2731, 2741, 
+	2749, 2753, 2767, 2777, 2789, 2791, 2797, 2801, 2803, 2819, 
+	2833, 2837, 2843, 2851, 2857, 2861, 2879, 2887, 2897, 2903, 
+	2909, 2917, 2927, 2939, 2953, 2957, 2963, 2969, 2971, 2999, 
+	3001, 3011, 3019, 3023, 3037, 3041, 3049, 3061, 3067, 3079, 
+	3083, 3089, 3109, 3119, 3121, 3137, 3163, 3167, 3169, 3181, 
+	3187, 3191, 3203, 3209, 3217, 3221, 3229, 3251, 3253, 3257, 
+	3259, 3271, 3299, 3301, 3307, 3313, 3319, 3323, 3329, 3331, 
+	3343, 3347, 3359, 3361, 3371, 3373, 3389, 3391, 3407, 3413, 
+	3433, 3449, 3457, 3461, 3463, 3467, 3469, 3491, 3499, 3511, 
+	3517, 3527, 3529, 3533, 3539, 3541, 3547, 3557, 3559, 3571, 
+	3581, 3583, 3593, 3607, 3613, 3617, 3623, 3631, 3637, 3643, 
+	3659, 3671, 3673, 3677, 3691, 3697, 3701, 3709, 3719, 3727, 
+	3733, 3739, 3761, 3767, 3769, 3779, 3793, 3797, 3803, 3821, 
+	3823, 3833, 3847, 3851, 3853, 3863, 3877, 3881, 3889, 3907, 
+	3911, 3917, 3919, 3923, 3929, 3931, 3943, 3947, 3967, 3989, 
+	4001, 4003, 4007, 4013, 4019, 4021, 4027, 4049, 4051, 4057, 
+	4073, 4079, 4091, 4093, 4099, 4111, 4127, 4129, 4133, 4139, 
+	4153, 4157, 4159, 4177, 4201, 4211, 4217, 4219, 4229, 4231, 
+	4241, 4243, 4253, 4259, 4261, 4271, 4273, 4283, 4289, 4297, 
+	4327, 4337, 4339, 4349, 4357, 4363, 4373, 4391, 4397, 4409, 
+	4421, 4423, 4441, 4447, 4451, 4457, 4463, 4481, 4483, 4493, 
+	4507, 4513, 4517, 4519, 4523, 4547, 4549, 4561, 4567, 4583, 
+	4591, 4597, 4603, 4621, 4637, 4639, 4643, 4649, 4651, 4657, 
+	4663, 4673, 4679, 4691, 4703, 4721, 4723, 4729, 4733, 4751, 
+	4759, 4783, 4787, 4789, 4793, 4799, 4801, 4813, 4817, 4831, 
+	4861, 4871, 4877, 4889, 4903, 4909, 4919, 4931, 4933, 4937, 
+	4943, 4951, 4957, 4967, 4969, 4973, 4987, 4993, 4999, 5003, 
+	5009, 5011, 5021, 5023, 5039, 5051, 5059, 5077, 5081, 5087, 
+	5099, 5101, 5107, 5113, 5119, 5147, 5153, 5167, 5171, 5179, 
+	5189, 5197, 5209, 5227, 5231, 5233, 5237, 5261, 5273, 5279, 
+	5281, 5297, 5303, 5309, 5323, 5333, 5347, 5351, 5381, 5387, 
+	5393, 5399, 5407, 5413, 5417, 5419, 5431, 5437, 5441, 5443, 
+	5449, 5471, 5477, 5479, 5483, 5501, 5503, 5507, 5519, 5521, 
+	5527, 5531, 5557, 5563, 5569, 5573, 5581, 5591, 5623, 5639, 
+	5641, 5647, 5651, 5653, 5657, 5659, 5669, 5683, 5689, 5693, 
+	5701, 5711, 5717, 5737, 5741, 5743, 5749, 5779, 5783, 5791, 
+	5801, 5807, 5813, 5821, 5827, 5839, 5843, 5849, 5851, 5857, 
+	5861, 5867, 5869, 5879, 5881, 5897, 5903, 5923, 5927, 5939, 
+	5953, 5981, 5987, 6007, 6011, 6029, 6037, 6043, 6047, 6053, 
+	6067, 6073, 6079, 6089, 6091, 6101, 6113, 6121, 6131, 6133, 
+	6143, 6151, 6163, 6173, 6197, 6199, 6203, 6211, 6217, 6221, 
+	6229, 6247, 6257, 6263, 6269, 6271, 6277, 6287, 6299, 6301, 
+	6311, 6317, 6323, 6329, 6337, 6343, 6353, 6359, 6361, 6367, 
+	6373, 6379, 6389, 6397, 6421, 6427, 6449, 6451, 6469, 6473, 
+	6481, 6491, 6521, 6529, 6547, 6551, 6553, 6563, 6569, 6571, 
+	6577, 6581, 6599, 6607, 6619, 6637, 6653, 6659, 6661, 6673, 
+	6679, 6689, 6691, 6701, 6703, 6709, 6719, 6733, 6737, 6761, 
+	6763, 6779, 6781, 6791, 6793, 6803, 6823, 6827, 6829, 6833, 
+	6841, 6857, 6863, 6869, 6871, 6883, 6899, 6907, 6911, 6917, 
+	6947, 6949, 6959, 6961, 6967, 6971, 6977, 6983, 6991, 6997, 
+	7001, 7013, 7019, 7027, 7039, 7043, 7057, 7069, 7079, 7103, 
+	7109, 7121, 7127, 7129, 7151, 7159, 7177, 7187, 7193, 7207, 
+	7211, 7213, 7219, 7229, 7237, 7243, 7247, 7253, 7283, 7297, 
+	7307, 7309, 7321, 7331, 7333, 7349, 7351, 7369, 7393, 7411, 
+	7417, 7433, 7451, 7457, 7459, 7477, 7481, 7487, 7489, 7499, 
+	7507, 7517, 7523, 7529, 7537, 7541, 7547, 7549, 7559, 7561, 
+	7573, 7577, 7583, 7589, 7591, 7603, 7607, 7621, 7639, 7643, 
+	7649, 7669, 7673, 7681, 7687, 7691, 7699, 7703, 7717, 7723, 
+	7727, 7741, 7753, 7757, 7759, 7789, 7793, 7817, 7823, 7829, 
+	7841, 7853, 7867, 7873, 7877, 7879, 7883, 7901, 7907, 7919, 
+	7927, 7933, 7937, 7949, 7951, 7963, 7993, 8009, 8011, 8017, 
+	8039, 8053, 8059, 8069, 8081, 8087, 8089, 8093, 8101, 8111, 
+	8117, 8123, 8147, 8161, 8167, 8171, 8179, 8191, 8209, 8219, 
+	8221, 8231, 8233, 8237, 8243, 8263, 8269, 8273, 8287, 8291, 
+	8293, 8297, 8311, 8317, 8329, 8353, 8363, 8369, 8377, 8387, 
+	8389, 8419, 8423, 8429, 8431, 8443, 8447, 8461, 8467, 8501, 
+	8513, 8521, 8527, 8537, 8539, 8543, 8563, 8573, 8581, 8597, 
+	8599, 8609, 8623, 8627, 8629, 8641, 8647, 8663, 8669, 8677, 
+	8681, 8689, 8693, 8699, 8707, 8713, 8719, 8731, 8737, 8741, 
+	8747, 8753, 8761, 8779, 8783, 8803, 8807, 8819, 8821, 8831, 
+	8837, 8839, 8849, 8861, 8863, 8867, 8887, 8893, 8923, 8929, 
+	8933, 8941, 8951, 8963, 8969, 8971, 8999, 9001, 9007, 9011, 
+	9013, 9029, 9041, 9043, 9049, 9059, 9067, 9091, 9103, 9109, 
+	9127, 9133, 9137, 9151, 9157, 9161, 9173, 9181, 9187, 9199, 
+	9203, 9209, 9221, 9227, 9239, 9241, 9257, 9277, 9281, 9283, 
+	9293, 9311, 9319, 9323, 9337, 9341, 9343, 9349, 9371, 9377, 
+	9391, 9397, 9403, 9413, 9419, 9421, 9431, 9433, 9437, 9439, 
+	9461, 9463, 9467, 9473, 9479, 9491, 9497, 9511, 9521, 9533, 
+	9539, 9547, 9551, 9587, 9601, 9613, 9619, 9623, 9629, 9631, 
+	9643, 9649, 9661, 9677, 9679, 9689, 9697, 9719, 9721, 9733, 
+	9739, 9743, 9749, 9767, 9769, 9781, 9787, 9791, 9803, 9811, 
+	9817, 9829, 9833, 9839, 9851, 9857, 9859, 9871, 9883, 9887, 
+	9901, 9907, 9923, 9929, 9931, 9941, 9949, 9967, 9973, 
+};
diff --git a/Engine/Source/Developer/PakFileUtilities/Private/SignedArchiveWriter.cpp b/Engine/Source/Developer/PakFileUtilities/Private/SignedArchiveWriter.cpp
new file mode 100644
--- /dev/null
+++ b/Engine/Source/Developer/PakFileUtilities/Private/SignedArchiveWriter.cpp
@@ -0,0 +1,162 @@
+// Copyright 1998-2018 Epic Games, Inc. All Rights Reserved.
+
+#include "SignedArchiveWriter.h"
+#include "IPlatformFilePak.h"
+#include "Misc/SecureHash.h"
+#include "HAL/FileManager.h"
+
+FSignedArchiveWriter::FSignedArchiveWriter(FArchive& InPak, const FString& InPakFilename, const FEncryptionKey& InPublicKey, const FEncryptionKey& InPrivateKey)
+: BufferArchive(Buffer)
+	, PakWriter(InPak)
+	, PakSignaturesFilename(FPaths::ChangeExtension(InPakFilename, TEXT("sig")))
+	, SizeOnDisk(0)
+	, PakSize(0)
+	, PublicKey(InPublicKey)
+	, PrivateKey(InPrivateKey)
+{
+	Buffer.Reserve(FPakInfo::MaxChunkDataSize);
+}
+
+FSignedArchiveWriter::~FSignedArchiveWriter()
+{
+	if (BufferArchive.Tell() > 0)
+	{
+		SerializeBufferAndSign();
+	}
+	delete &PakWriter;
+}
+
+void FSignedArchiveWriter::SerializeBufferAndSign()
+{
+	// Compute a hash for this buffer data
+	ChunkHashes.Add(ComputePakChunkHash(&Buffer[0], Buffer.Num()));
+
+	// Flush the buffer
+	PakWriter.Serialize(&Buffer[0], Buffer.Num());
+	BufferArchive.Seek(0);
+	Buffer.Empty(FPakInfo::MaxChunkDataSize);
+}
+
+bool FSignedArchiveWriter::Close()
+{
+	if (BufferArchive.Tell() > 0)
+	{
+		SerializeBufferAndSign();
+	}
+
+	FEncryptedSignature EncryptedMasterHash;
+	FDecryptedSignature DecryptedMasterHash;
+	DecryptedMasterHash.Data = ComputePakChunkHash((const uint8*)&ChunkHashes[0], ChunkHashes.Num() * sizeof(TPakChunkHash));
+	FEncryption::EncryptSignature(DecryptedMasterHash, EncryptedMasterHash, PrivateKey);
+
+	FArchive* SignatureWriter = IFileManager::Get().CreateFileWriter(*PakSignaturesFilename);
+	*SignatureWriter << EncryptedMasterHash;
+	*SignatureWriter << ChunkHashes;
+	delete SignatureWriter;
+
+	return FArchive::Close();
+}
+
+void FSignedArchiveWriter::Serialize(void* Data, int64 Length)
+{
+	// Serialize data to a buffer. When the max buffer size is reached, the buffer is signed and
+	// serialized to disk with its signature
+	uint8* DataToWrite = (uint8*)Data;
+	int64 RemainingSize = Length;
+	while (RemainingSize > 0)
+	{
+		int64 BufferPos = BufferArchive.Tell();
+		int64 SizeToWrite = RemainingSize;
+		if (BufferPos + SizeToWrite > FPakInfo::MaxChunkDataSize)
+		{
+			SizeToWrite = FPakInfo::MaxChunkDataSize - BufferPos;
+		}
+
+		BufferArchive.Serialize(DataToWrite, SizeToWrite);
+		if (BufferArchive.Tell() == FPakInfo::MaxChunkDataSize)
+		{
+			SerializeBufferAndSign();
+		}
+			
+		SizeOnDisk += SizeToWrite;
+		PakSize += SizeToWrite;
+
+		RemainingSize -= SizeToWrite;
+		DataToWrite += SizeToWrite;
+	}
+}
+
+int64 FSignedArchiveWriter::Tell()
+{
+	return PakSize;
+}
+
+int64 FSignedArchiveWriter::TotalSize()
+{
+	return PakSize;
+}
+
+void FSignedArchiveWriter::Seek(int64 InPos)
+{
+	UE_LOG(LogPakFile, Fatal, TEXT("Seek is not supported in FSignedArchiveWriter."));
+}
+
+/**
+ * Useful code for testing the encryption methods
+ */
+void TestEncryption()
+{
+	FEncryptionKey PublicKey;
+	FEncryptionKey PrivateKey;
+	TEncryptionInt P(TEXT("0x21443BD2DD63E995403"));
+	TEncryptionInt Q(TEXT("0x28CBB6E5749AC65749"));
+	FEncryption::GenerateKeyPair(P, Q, PublicKey, PrivateKey);
+
+	// Generate random data
+	const int32 DataSize = 1024;
+	uint8* Data = (uint8*)FMemory::Malloc(DataSize);
+	for (int32 Index = 0; Index < DataSize; ++Index)
+	{
+		Data[Index] = (uint8)(Index % 255);
+	}
+
+	// Generate signature
+	FDecryptedSignature OriginalSignature;
+	FEncryptedSignature EncryptedSignature;
+	FDecryptedSignature DecryptedSignature;
+	OriginalSignature.Data = FCrc::MemCrc32(Data, DataSize);
+
+	// Encrypt signature with the private key
+	FEncryption::EncryptSignature(OriginalSignature, EncryptedSignature, PrivateKey);
+	FEncryption::DecryptSignature(EncryptedSignature, DecryptedSignature, PublicKey);
+
+	// Check
+	if (OriginalSignature == DecryptedSignature)
+	{
+		UE_LOG(LogPakFile, Display, TEXT("Keys match"));
+	}
+	else
+	{
+		UE_LOG(LogPakFile, Fatal, TEXT("Keys mismatched!"));
+	}
+
+	double OverallTime = 0.0;
+	double OverallNumTests = 0.0;
+	for (int32 TestIndex = 0; TestIndex < 10; ++TestIndex)
+	{
+		static const int64 NumTests = 500;
+		double Timer = FPlatformTime::Seconds();
+		{
+			for (int64 i = 0; i < NumTests; ++i)
+			{
+				FEncryption::DecryptSignature(EncryptedSignature, DecryptedSignature, PublicKey);
+			}
+		}
+		Timer = FPlatformTime::Seconds() - Timer;
+		OverallTime += Timer;
+		OverallNumTests += (double)NumTests;
+		UE_LOG(LogPakFile, Display, TEXT("%i signatures decrypted in %.4fs, Avg = %.4fs, OverallAvg = %.4fs"), NumTests, Timer, Timer / (double)NumTests, OverallTime / OverallNumTests);
+	}
+	
+	FMemory::Free(Data);
+}
diff --git a/Engine/Source/Developer/PakFileUtilities/Private/SignedArchiveWriter.h b/Engine/Source/Developer/PakFileUtilities/Private/SignedArchiveWriter.h
new file mode 100644
--- /dev/null
+++ b/Engine/Source/Developer/PakFileUtilities/Private/SignedArchiveWriter.h
@@ -0,0 +1,50 @@
+// Copyright 1998-2018 Epic Games, Inc. All Rights Reserved.
+#pragma once
+
+#include "CoreMinimal.h"
+#include "Serialization/Archive.h"
+#include "Serialization/MemoryWriter.h"
+#include "Math/BigInt.h"
+#include "IPlatformFilePak.h"
+
+/**
+ * Wrapper for writing and signing an archive
+ */
+class FSignedArchiveWriter : public FArchive
+{
+	/** Buffer to sign */
+	TArray<uint8> Buffer;
+	/** Buffer writer */
+	FMemoryWriter BufferArchive;
+	/** The actual pak archive */
+	FArchive& PakWriter;
+	/** The filename of the signature file that accompanies the pak */
+	FString PakSignaturesFilename;
+	/** Size of the archive on disk (including signatures) */
+	int64 SizeOnDisk;
+	/** Data size (excluding signatures) */
+	int64 PakSize;
+	/** Decryption key */
+	FEncryptionKey PublicKey;
+	/** Encryption key */
+	FEncryptionKey PrivateKey;
+	/** Hashes */
+	TArray<TPakChunkHash> ChunkHashes;
+
+	/** 
+	 * Serializes and signs a buffer
+	 */
+	void SerializeBufferAndSign();
+
+public:
+
+	FSignedArchiveWriter(FArchive& InPak, const FString& InPakFilename, const FEncryptionKey& InPublicKey, const FEncryptionKey& InPrivateKey);
+	virtual ~FSignedArchiveWriter();
+
+	// FArchive interface
+	virtual bool Close() override;
+	virtual void Serialize(void* Data, int64 Length) override;
+	virtual int64 Tell() override;
+	virtual int64 TotalSize() override;
+	virtual void Seek(int64 InPos) override;
+};
diff --git a/Engine/Source/Developer/PakFileUtilities/Public/PakFileUtilities.h b/Engine/Source/Developer/PakFileUtilities/Public/PakFileUtilities.h
new file mode 100644
--- /dev/null
+++ b/Engine/Source/Developer/PakFileUtilities/Public/PakFileUtilities.h
@@ -0,0 +1,6 @@
+// Copyright 1998-2018 Epic Games, Inc. All Rights Reserved.
+
+#include "CoreMinimal.h"
+
+PAKFILEUTILITIES_API int32 ExecuteUnrealPak(int ArgC, TCHAR* ArgV[]);
+
diff --git a/Engine/Source/Programs/UnrealPak/Private/KeyGenerator.cpp b/Engine/Source/Programs/UnrealPak/Private/KeyGenerator.cpp
deleted file mode 100644
diff --git a/Engine/Source/Programs/UnrealPak/Private/KeyGenerator.h b/Engine/Source/Programs/UnrealPak/Private/KeyGenerator.h
deleted file mode 100644
diff --git a/Engine/Source/Programs/UnrealPak/Private/Primes.inl b/Engine/Source/Programs/UnrealPak/Private/Primes.inl
deleted file mode 100644
diff --git a/Engine/Source/Programs/UnrealPak/Private/SignedArchiveWriter.cpp b/Engine/Source/Programs/UnrealPak/Private/SignedArchiveWriter.cpp
deleted file mode 100644
diff --git a/Engine/Source/Programs/UnrealPak/Private/SignedArchiveWriter.h b/Engine/Source/Programs/UnrealPak/Private/SignedArchiveWriter.h
deleted file mode 100644
