Missing header.

diff --git a/Engine/Source/Developer/PakFileUtilities/Private/PakFileUtilities.cpp b/Engine/Source/Developer/PakFileUtilities/Private/PakFileUtilities.cpp
--- a/Engine/Source/Developer/PakFileUtilities/Private/PakFileUtilities.cpp
+++ b/Engine/Source/Developer/PakFileUtilities/Private/PakFileUtilities.cpp
@@ -16,6 +16,7 @@
 #include "Misc/Compression.h"
 #include "Features/IModularFeatures.h"
 #include "Misc/CoreDelegates.h"
+#include "Misc/CommandLine.h"
 
 struct FPakCommandLineParameters
 {
