Fixes for compiling PakFileUtilities as part of the editor.

#rnx

diff --git a/Engine/Source/Developer/PakFileUtilities/Private/PakFileUtilities.cpp b/Engine/Source/Developer/PakFileUtilities/Private/PakFileUtilities.cpp
--- a/Engine/Source/Developer/PakFileUtilities/Private/PakFileUtilities.cpp
+++ b/Engine/Source/Developer/PakFileUtilities/Private/PakFileUtilities.cpp
@@ -17,6 +17,9 @@
 #include "Features/IModularFeatures.h"
 #include "Misc/CoreDelegates.h"
 #include "Misc/CommandLine.h"
+#include "Misc/FileHelper.h"
+#include "Misc/ConfigCacheIni.h"
+#include "HAL/PlatformFilemanager.h"
 
 struct FPakCommandLineParameters
 {
@@ -1508,7 +1511,7 @@ bool CreatePakFile(const TCHAR* Filename, TArray<FPakInputPair>& FilesToAdd, con
 
 bool TestPakFile(const TCHAR* Filename)
 {	
-	FPakFile PakFile(Filename, FParse::Param(FCommandLine::Get(), TEXT("signed")));
+	FPakFile PakFile(&FPlatformFileManager::Get().GetPlatformFile(), Filename, FParse::Param(FCommandLine::Get(), TEXT("signed")));
 	if (PakFile.IsValid())
 	{
 		return PakFile.Check();
@@ -1522,7 +1525,7 @@ bool TestPakFile(const TCHAR* Filename)
 
 bool ListFilesInPak(const TCHAR * InPakFilename, int64 SizeFilter = 0)
 {
-	FPakFile PakFile(InPakFilename, FParse::Param(FCommandLine::Get(), TEXT("signed")));
+	FPakFile PakFile(&FPlatformFileManager::Get().GetPlatformFile(), InPakFilename, FParse::Param(FCommandLine::Get(), TEXT("signed")));
 	int32 FileCount = 0;
 	int64 FileSize = 0;
 	int64 FilteredSize = 0;
@@ -1612,7 +1615,7 @@ bool ExtractFilesFromPak(const TCHAR* InPakFilename, TMap<FString, FFileInfo>& I
 			}
 		}
 
-		FPakFile PakFile(*PakFilename, FParse::Param(FCommandLine::Get(), TEXT("signed")));
+		FPakFile PakFile(&FPlatformFileManager::Get().GetPlatformFile(), *PakFilename, FParse::Param(FCommandLine::Get(), TEXT("signed")));
 		if (PakFile.IsValid())
 		{
 			FString DestPath(InDestPath);
@@ -1729,8 +1732,8 @@ bool DiffFilesInPaks(const FString InPakFilename1, const FString InPakFilename2,
 	const bool bLogUniques1 = bLogUniques && !FParse::Param(FCommandLine::Get(), TEXT("nouniquesfile1"));
 	const bool bLogUniques2 = bLogUniques && !FParse::Param(FCommandLine::Get(), TEXT("nouniquesfile2"));
 
-	FPakFile PakFile1(*InPakFilename1, FParse::Param(FCommandLine::Get(), TEXT("signed")));
-	FPakFile PakFile2(*InPakFilename2, FParse::Param(FCommandLine::Get(), TEXT("signed")));
+	FPakFile PakFile1(&FPlatformFileManager::Get().GetPlatformFile(), *InPakFilename1, FParse::Param(FCommandLine::Get(), TEXT("signed")));
+	FPakFile PakFile2(&FPlatformFileManager::Get().GetPlatformFile(), *InPakFilename2, FParse::Param(FCommandLine::Get(), TEXT("signed")));
 	if (PakFile1.IsValid() && PakFile2.IsValid())
 	{		
 		FArchive& PakReader1 = *PakFile1.GetSharedReader(NULL);
@@ -1929,7 +1932,7 @@ bool GenerateHashesFromPak(const TCHAR* InPakFilename, const TCHAR* InDestPakFil
 			}
 		}
 
-		FPakFile PakFile(*PakFilename, FParse::Param(FCommandLine::Get(), TEXT("signed")));
+		FPakFile PakFile(&FPlatformFileManager::Get().GetPlatformFile(), *PakFilename, FParse::Param(FCommandLine::Get(), TEXT("signed")));
 		if (PakFile.IsValid())
 		{
 			FArchive& PakReader = *PakFile.GetSharedReader(NULL);
@@ -2274,7 +2277,7 @@ bool ExportDependencies(const TCHAR * PakFilename, const TCHAR* GameName, const
 	// Example command line used for this tool
 	// C:\Development\BB\WEX\Saved\StagedBuilds\WindowsNoEditor\WorldExplorers\Content\Paks\WorldExplorers-WindowsNoEditor.pak WorldExplorers WEX -exportdependencies=c:\dvtemp\output -debug -NoAssetRegistryCache -ForceDependsGathering
 	
-	FPakFile PakFile(PakFilename,FParse::Param(FCommandLine::Get(),TEXT("signed")));
+	FPakFile PakFile(&FPlatformFileManager::Get().GetPlatformFile(), PakFilename,FParse::Param(FCommandLine::Get(),TEXT("signed")));
 
 	if(PakFile.IsValid())
 	{
