Add an "UnrealPak" commandlet with the same functionality as the standalone UnrealPak program. Invoke by running the editor with -run=UnrealPak and the standard UnrealPak commandline options.

diff --git a/Engine/Source/Developer/PakFileUtilities/Private/PakFileUtilities.cpp b/Engine/Source/Developer/PakFileUtilities/Private/PakFileUtilities.cpp
--- a/Engine/Source/Developer/PakFileUtilities/Private/PakFileUtilities.cpp
+++ b/Engine/Source/Developer/PakFileUtilities/Private/PakFileUtilities.cpp
@@ -16,7 +16,6 @@
 #include "Misc/Compression.h"
 #include "Features/IModularFeatures.h"
 #include "Misc/CoreDelegates.h"
-#include "Misc/CommandLine.h"
 #include "Misc/FileHelper.h"
 #include "Misc/ConfigCacheIni.h"
 #include "HAL/PlatformFilemanager.h"
@@ -356,56 +355,54 @@ bool PrepareCopyCompressedFileToPak(const FString& InMountPoint, const FPakInput
 	return true;
 }
 
-void ProcessOrderFile(int32 ArgC, TCHAR* ArgV[], TMap<FString, uint64>& OrderMap)
+bool ProcessOrderFile(const TCHAR* ResponseFile, TMap<FString, uint64>& OrderMap)
 {
 	// List of all items to add to pak file
-	FString ResponseFile;
-	if (FParse::Value(FCommandLine::Get(), TEXT("-order="), ResponseFile))
+	FString Text;
+	UE_LOG(LogPakFile, Display, TEXT("Loading pak order file %s..."), ResponseFile);
+	if (FFileHelper::LoadFileToString(Text, ResponseFile))
 	{
-		FString Text;
-		UE_LOG(LogPakFile, Display, TEXT("Loading pak order file %s..."), *ResponseFile);
-		if (FFileHelper::LoadFileToString(Text, *ResponseFile))
+		// Read all lines
+		TArray<FString> Lines;
+		Text.ParseIntoArray(Lines, TEXT("\n"), true);
+		for (int32 EntryIndex = 0; EntryIndex < Lines.Num(); EntryIndex++)
 		{
-			// Read all lines
-			TArray<FString> Lines;
-			Text.ParseIntoArray(Lines, TEXT("\n"), true);
-			for (int32 EntryIndex = 0; EntryIndex < Lines.Num(); EntryIndex++)
+			Lines[EntryIndex].ReplaceInline(TEXT("\r"), TEXT(""));
+			Lines[EntryIndex].ReplaceInline(TEXT("\n"), TEXT(""));
+			int32 OpenOrderNumber = EntryIndex;
+			if (Lines[EntryIndex].FindLastChar('"', OpenOrderNumber))
 			{
-				Lines[EntryIndex].ReplaceInline(TEXT("\r"), TEXT(""));
-				Lines[EntryIndex].ReplaceInline(TEXT("\n"), TEXT(""));
-				int32 OpenOrderNumber = EntryIndex;
-				if (Lines[EntryIndex].FindLastChar('"', OpenOrderNumber))
+				FString ReadNum = Lines[EntryIndex].RightChop(OpenOrderNumber+1);
+				Lines[EntryIndex] = Lines[EntryIndex].Left(OpenOrderNumber+1);
+				ReadNum.TrimStartInline();
+				if (ReadNum.IsNumeric())
 				{
-					FString ReadNum = Lines[EntryIndex].RightChop(OpenOrderNumber+1);
-					Lines[EntryIndex] = Lines[EntryIndex].Left(OpenOrderNumber+1);
-					ReadNum.TrimStartInline();
-					if (ReadNum.IsNumeric())
-					{
-						OpenOrderNumber = FCString::Atoi(*ReadNum);
-					}
+					OpenOrderNumber = FCString::Atoi(*ReadNum);
 				}
-				Lines[EntryIndex] = Lines[EntryIndex].TrimQuotes();
-				FString Path=FString::Printf(TEXT("%s"), *Lines[EntryIndex]);
-				FPaths::NormalizeFilename(Path);
-				Path = Path.ToLower();
+			}
+			Lines[EntryIndex] = Lines[EntryIndex].TrimQuotes();
+			FString Path=FString::Printf(TEXT("%s"), *Lines[EntryIndex]);
+			FPaths::NormalizeFilename(Path);
+			Path = Path.ToLower();
 #if 0
-				if (Path.EndsWith("uexp"))
-				{
-					OpenOrderNumber += (1 << 29);
-				}
-				if (Path.EndsWith("ubulk"))
-				{
-					OpenOrderNumber += (1 << 30);
-				}
-#endif
-				OrderMap.Add(Path, OpenOrderNumber);
+			if (Path.EndsWith("uexp"))
+			{
+				OpenOrderNumber += (1 << 29);
 			}
-			UE_LOG(LogPakFile, Display, TEXT("Finished loading pak order file %s."), *ResponseFile);
-		}
-		else 
-		{
-			UE_LOG(LogPakFile, Display, TEXT("Unable to load pak order file %s."), *ResponseFile);
+			if (Path.EndsWith("ubulk"))
+			{
+				OpenOrderNumber += (1 << 30);
+			}
+#endif
+			OrderMap.Add(Path, OpenOrderNumber);
 		}
+		UE_LOG(LogPakFile, Display, TEXT("Finished loading pak order file %s."), ResponseFile);
+		return true;
+	}
+	else 
+	{
+		UE_LOG(LogPakFile, Error, TEXT("Unable to load pak order file %s."), ResponseFile);
+		return false;
 	}
 }
 
@@ -425,14 +422,14 @@ static void CommandLineParseHelper(const TCHAR* InCmdLine, TArray<FString>& Toke
 	}
 }
 
-void ProcessCommandLine(int32 ArgC, TCHAR* ArgV[], TArray<FPakInputPair>& Entries, FPakCommandLineParameters& CmdLineParameters)
+void ProcessCommandLine(const TCHAR* CmdLine, const TArray<FString>& NonOptionArguments, TArray<FPakInputPair>& Entries, FPakCommandLineParameters& CmdLineParameters)
 {
 	// List of all items to add to pak file
 	FString ResponseFile;
 	FString ClusterSizeString;
 
-	if (FParse::Value(FCommandLine::Get(), TEXT("-blocksize="), ClusterSizeString) && 
-		FParse::Value(FCommandLine::Get(), TEXT("-blocksize="), CmdLineParameters.FileSystemBlockSize))
+	if (FParse::Value(CmdLine, TEXT("-blocksize="), ClusterSizeString) && 
+		FParse::Value(CmdLine, TEXT("-blocksize="), CmdLineParameters.FileSystemBlockSize))
 	{
 		if (ClusterSizeString.EndsWith(TEXT("MB")))
 		{
@@ -449,8 +446,8 @@ void ProcessCommandLine(int32 ArgC, TCHAR* ArgV[], TArray<FPakInputPair>& Entrie
 	}
 
 	FString CompBlockSizeString;
-	if (FParse::Value(FCommandLine::Get(), TEXT("-compressionblocksize="), CompBlockSizeString) &&
-		FParse::Value(FCommandLine::Get(), TEXT("-compressionblocksize="), CmdLineParameters.CompressionBlockSize))
+	if (FParse::Value(CmdLine, TEXT("-compressionblocksize="), CompBlockSizeString) &&
+		FParse::Value(CmdLine, TEXT("-compressionblocksize="), CmdLineParameters.CompressionBlockSize))
 	{
 		if (CompBlockSizeString.EndsWith(TEXT("MB")))
 		{
@@ -462,23 +459,23 @@ void ProcessCommandLine(int32 ArgC, TCHAR* ArgV[], TArray<FPakInputPair>& Entrie
 		}
 	}
 
-	if (!FParse::Value(FCommandLine::Get(), TEXT("-bitwindow="), CmdLineParameters.CompressionBitWindow))
+	if (!FParse::Value(CmdLine, TEXT("-bitwindow="), CmdLineParameters.CompressionBitWindow))
 	{
 		CmdLineParameters.CompressionBitWindow = DEFAULT_ZLIB_BIT_WINDOW;
 	}
 
-	if (!FParse::Value(FCommandLine::Get(), TEXT("-patchpaddingalign="), CmdLineParameters.PatchFilePadAlign))
+	if (!FParse::Value(CmdLine, TEXT("-patchpaddingalign="), CmdLineParameters.PatchFilePadAlign))
 	{
 		CmdLineParameters.PatchFilePadAlign = 0;
 	}
 
-	if (FParse::Param(FCommandLine::Get(), TEXT("encryptindex")))
+	if (FParse::Param(CmdLine, TEXT("encryptindex")))
 	{
 		CmdLineParameters.EncryptIndex = true;
 	}
 
 	FString CompressorFileName;
-	if (FParse::Value(FCommandLine::Get(), TEXT("compressor="), CompressorFileName))
+	if (FParse::Value(CmdLine, TEXT("compressor="), CompressorFileName))
 	{
 		void* CustomCompressorDll = FPlatformProcess::GetDllHandle(*CompressorFileName);
 		if (CustomCompressorDll == nullptr)
@@ -496,7 +493,7 @@ void ProcessCommandLine(int32 ArgC, TCHAR* ArgV[], TArray<FPakInputPair>& Entrie
 			return;
 		}
 
-		ICustomCompressor* Compressor = (*CreateCustomCompressor)(FCommandLine::Get());
+		ICustomCompressor* Compressor = (*CreateCustomCompressor)(CmdLine);
 		if (Compressor == nullptr)
 		{
 			UE_LOG(LogPakFile, Error, TEXT("Failed to create custom compressor from '%s'"), *CompressorFileName);
@@ -507,19 +504,19 @@ void ProcessCommandLine(int32 ArgC, TCHAR* ArgV[], TArray<FPakInputPair>& Entrie
 		CmdLineParameters.UseCustomCompressor = true;
 	}
 
-	if (FParse::Param(FCommandLine::Get(), TEXT("overrideplatformcompressor")))
+	if (FParse::Param(CmdLine, TEXT("overrideplatformcompressor")))
 	{
 		CmdLineParameters.OverridePlatformCompressor = true;
 	}
 
-	if (FParse::Value(FCommandLine::Get(), TEXT("-create="), ResponseFile))
+	if (FParse::Value(CmdLine, TEXT("-create="), ResponseFile))
 	{
 		TArray<FString> Lines;
 
-		CmdLineParameters.GeneratePatch = FParse::Value(FCommandLine::Get(), TEXT("-generatepatch="), CmdLineParameters.SourcePatchPakFilename);
+		CmdLineParameters.GeneratePatch = FParse::Value(CmdLine, TEXT("-generatepatch="), CmdLineParameters.SourcePatchPakFilename);
 
-		bool bCompress = FParse::Param(FCommandLine::Get(), TEXT("compress"));
-		bool bEncrypt = FParse::Param(FCommandLine::Get(), TEXT("encrypt"));
+		bool bCompress = FParse::Param(CmdLine, TEXT("compress"));
+		bool bEncrypt = FParse::Param(CmdLine, TEXT("encrypt"));
 
 		bool bParseLines = true;
 		if (IFileManager::Get().DirectoryExists(*ResponseFile))
@@ -599,34 +596,30 @@ void ProcessCommandLine(int32 ArgC, TCHAR* ArgV[], TArray<FPakInputPair>& Entrie
 	{
 		// Override destination path.
 		FString MountPoint;
-		FParse::Value(FCommandLine::Get(), TEXT("-dest="), MountPoint);
+		FParse::Value(CmdLine, TEXT("-dest="), MountPoint);
 		FPaths::NormalizeFilename(MountPoint);
 		FPakFile::MakeDirectoryFromPath(MountPoint);
 
 		// Parse comand line params. The first param after the program name is the created pak name
-		for (int32 Index = 2; Index < ArgC; Index++)
+		for (int32 Index = 1; Index < NonOptionArguments.Num(); Index++)
 		{
 			// Skip switches and add everything else to the Entries array
-			TCHAR* Param = ArgV[Index];
-			if (Param[0] != '-')
+			FPakInputPair Input;
+			Input.Source = *NonOptionArguments[Index];
+			FPaths::NormalizeFilename(Input.Source);
+			if (MountPoint.Len() > 0)
 			{
-				FPakInputPair Input;
-				Input.Source = Param;
-				FPaths::NormalizeFilename(Input.Source);
-				if (MountPoint.Len() > 0)
-				{
-					FString SourceDirectory( FPaths::GetPath(Input.Source) );
-					FPakFile::MakeDirectoryFromPath(SourceDirectory);
-					Input.Dest = Input.Source.Replace(*SourceDirectory, *MountPoint, ESearchCase::IgnoreCase);
-				}
-				else
-				{
-					Input.Dest = FPaths::GetPath(Input.Source);
-					FPakFile::MakeDirectoryFromPath(Input.Dest);
-				}
-				FPaths::NormalizeFilename(Input.Dest);
-				Entries.Add(Input);
+				FString SourceDirectory( FPaths::GetPath(Input.Source) );
+				FPakFile::MakeDirectoryFromPath(SourceDirectory);
+				Input.Dest = Input.Source.Replace(*SourceDirectory, *MountPoint, ESearchCase::IgnoreCase);
 			}
+			else
+			{
+				Input.Dest = FPaths::GetPath(Input.Source);
+				FPakFile::MakeDirectoryFromPath(Input.Dest);
+			}
+			FPaths::NormalizeFilename(Input.Dest);
+			Entries.Add(Input);
 		}
 	}
 	UE_LOG(LogPakFile, Display, TEXT("Added %d entries to add to pak file."), Entries.Num());
@@ -877,7 +870,7 @@ void PrepareEncryptionAndSigningKeysFromCryptoKeyCache(const FString& InFilename
 	delete File;
 }
 
-void PrepareEncryptionAndSigningKeys(FKeyPair& OutSigningKey, FAES::FAESKey& OutAESKey)
+void PrepareEncryptionAndSigningKeys(const TCHAR* CmdLine, FKeyPair& OutSigningKey, FAES::FAESKey& OutAESKey)
 {
 	OutSigningKey.PrivateKey.Exponent.Zero();
 	OutSigningKey.PrivateKey.Modulus.Zero();
@@ -887,18 +880,18 @@ void PrepareEncryptionAndSigningKeys(FKeyPair& OutSigningKey, FAES::FAESKey& Out
 
 	// First, try and parse the keys from a supplied crypto key cache file
 	FString CryptoKeysCacheFilename;
-	if (FParse::Value(FCommandLine::Get(), TEXT("cryptokeys="), CryptoKeysCacheFilename))
+	if (FParse::Value(CmdLine, TEXT("cryptokeys="), CryptoKeysCacheFilename))
 	{
 		UE_LOG(LogPakFile, Display, TEXT("Parsing crypto keys from a crypto key cache file"));
 		PrepareEncryptionAndSigningKeysFromCryptoKeyCache(CryptoKeysCacheFilename, OutSigningKey, OutAESKey);
 	}
-	else if (FParse::Param(FCommandLine::Get(), TEXT("encryptionini")))
+	else if (FParse::Param(CmdLine, TEXT("encryptionini")))
 	{
 		FString ProjectDir, EngineDir, Platform;
 
-		if (FParse::Value(FCommandLine::Get(), TEXT("projectdir="), ProjectDir, false)
-			&& FParse::Value(FCommandLine::Get(), TEXT("enginedir="), EngineDir, false)
-			&& FParse::Value(FCommandLine::Get(), TEXT("platform="), Platform, false))
+		if (FParse::Value(CmdLine, TEXT("projectdir="), ProjectDir, false)
+			&& FParse::Value(CmdLine, TEXT("enginedir="), EngineDir, false)
+			&& FParse::Value(CmdLine, TEXT("platform="), Platform, false))
 		{
 			FConfigFile EngineConfig;
 
@@ -1014,7 +1007,7 @@ void PrepareEncryptionAndSigningKeys(FKeyPair& OutSigningKey, FAES::FAESKey& Out
 		UE_LOG(LogPakFile, Display, TEXT("Using command line for crypto configuration"));
 
 		FString EncryptionKeyString;
-		FParse::Value(FCommandLine::Get(), TEXT("aes="), EncryptionKeyString, false);
+		FParse::Value(CmdLine, TEXT("aes="), EncryptionKeyString, false);
 
 		if (EncryptionKeyString.Len() > 0)
 		{
@@ -1044,7 +1037,7 @@ void PrepareEncryptionAndSigningKeys(FKeyPair& OutSigningKey, FAES::FAESKey& Out
 		}
 
 		FString KeyFilename;
-		if (FParse::Value(FCommandLine::Get(), TEXT("sign="), KeyFilename, false))
+		if (FParse::Value(CmdLine, TEXT("sign="), KeyFilename, false))
 		{
 			if (KeyFilename.StartsWith(TEXT("0x")))
 			{
@@ -1509,9 +1502,9 @@ bool CreatePakFile(const TCHAR* Filename, TArray<FPakInputPair>& FilesToAdd, con
 	return true;
 }
 
-bool TestPakFile(const TCHAR* Filename)
+bool TestPakFile(const TCHAR* Filename, bool bSigned)
 {	
-	FPakFile PakFile(&FPlatformFileManager::Get().GetPlatformFile(), Filename, FParse::Param(FCommandLine::Get(), TEXT("signed")));
+	FPakFile PakFile(&FPlatformFileManager::Get().GetPlatformFile(), Filename, bSigned);
 	if (PakFile.IsValid())
 	{
 		return PakFile.Check();
@@ -1523,9 +1516,9 @@ bool TestPakFile(const TCHAR* Filename)
 	}
 }
 
-bool ListFilesInPak(const TCHAR * InPakFilename, int64 SizeFilter = 0)
+bool ListFilesInPak(const TCHAR *InPakFilename, int64 SizeFilter, bool bSigned)
 {
-	FPakFile PakFile(&FPlatformFileManager::Get().GetPlatformFile(), InPakFilename, FParse::Param(FCommandLine::Get(), TEXT("signed")));
+	FPakFile PakFile(&FPlatformFileManager::Get().GetPlatformFile(), InPakFilename, bSigned);
 	int32 FileCount = 0;
 	int64 FileSize = 0;
 	int64 FilteredSize = 0;
@@ -1580,7 +1573,7 @@ struct FFileInfo
 	uint8 Hash[16];
 };
 
-bool ExtractFilesFromPak(const TCHAR* InPakFilename, TMap<FString, FFileInfo>& InFileHashes, const TCHAR* InDestPath, bool bUseMountPoint, const FAES::FAESKey& InEncryptionKey, TArray<FPakInputPair>* OutEntries = nullptr, TMap<FString, uint64>* OutOrderMap = nullptr)
+bool ExtractFilesFromPak(const TCHAR* InPakFilename, TMap<FString, FFileInfo>& InFileHashes, const TCHAR* InDestPath, bool bUseMountPoint, const FAES::FAESKey& InEncryptionKey, bool bSigned, TArray<FPakInputPair>* OutEntries = nullptr, TMap<FString, uint64>* OutOrderMap = nullptr)
 {
 	// Gather all patch versions of the requested pak file and run through each separately
 	TArray<FString> PakFileList;
@@ -1615,7 +1608,7 @@ bool ExtractFilesFromPak(const TCHAR* InPakFilename, TMap<FString, FFileInfo>& I
 			}
 		}
 
-		FPakFile PakFile(&FPlatformFileManager::Get().GetPlatformFile(), *PakFilename, FParse::Param(FCommandLine::Get(), TEXT("signed")));
+		FPakFile PakFile(&FPlatformFileManager::Get().GetPlatformFile(), *PakFilename, bSigned);
 		if (PakFile.IsValid())
 		{
 			FString DestPath(InDestPath);
@@ -1717,7 +1710,7 @@ void CreateDiffRelativePathMap(TArray<FString>& FileNames, const FString& RootPa
 	}
 }
 
-bool DiffFilesInPaks(const FString InPakFilename1, const FString InPakFilename2, const FAES::FAESKey& InEncryptionKey)
+bool DiffFilesInPaks(const FString InPakFilename1, const FString InPakFilename2, bool bLogUniques1, bool bLogUniques2, const FAES::FAESKey& InEncryptionKey, bool bSigned)
 {
 	int32 NumUniquePAK1 = 0;
 	int32 NumUniquePAK2 = 0;
@@ -1727,13 +1720,8 @@ bool DiffFilesInPaks(const FString InPakFilename1, const FString InPakFilename2,
 	TGuardValue<ELogTimes::Type> DisableLogTimes(GPrintLogTimes, ELogTimes::None);
 	UE_LOG(LogPakFile, Log, TEXT("FileEventType, FileName, Size1, Size2"));
 
-	// Allow the suppression of unique file logging for one or both files
-	const bool bLogUniques = !FParse::Param(FCommandLine::Get(), TEXT("nouniques"));
-	const bool bLogUniques1 = bLogUniques && !FParse::Param(FCommandLine::Get(), TEXT("nouniquesfile1"));
-	const bool bLogUniques2 = bLogUniques && !FParse::Param(FCommandLine::Get(), TEXT("nouniquesfile2"));
-
-	FPakFile PakFile1(&FPlatformFileManager::Get().GetPlatformFile(), *InPakFilename1, FParse::Param(FCommandLine::Get(), TEXT("signed")));
-	FPakFile PakFile2(&FPlatformFileManager::Get().GetPlatformFile(), *InPakFilename2, FParse::Param(FCommandLine::Get(), TEXT("signed")));
+	FPakFile PakFile1(&FPlatformFileManager::Get().GetPlatformFile(), *InPakFilename1, bSigned);
+	FPakFile PakFile2(&FPlatformFileManager::Get().GetPlatformFile(), *InPakFilename2, bSigned);
 	if (PakFile1.IsValid() && PakFile2.IsValid())
 	{		
 		FArchive& PakReader1 = *PakFile1.GetSharedReader(NULL);
@@ -1894,7 +1882,7 @@ bool GenerateHashForFile( FString Filename, FFileInfo& FileHash)
 	return true;
 }
 
-bool GenerateHashesFromPak(const TCHAR* InPakFilename, const TCHAR* InDestPakFilename, TMap<FString, FFileInfo>& FileHashes, bool bUseMountPoint, const FAES::FAESKey& InEncryptionKey)
+bool GenerateHashesFromPak(const TCHAR* InPakFilename, const TCHAR* InDestPakFilename, TMap<FString, FFileInfo>& FileHashes, bool bUseMountPoint, const FAES::FAESKey& InEncryptionKey, bool bSigned)
 {
 	TArray<FString> FoundFiles;
 	IFileManager::Get().FindFiles(FoundFiles, InPakFilename, true, false);
@@ -1932,7 +1920,7 @@ bool GenerateHashesFromPak(const TCHAR* InPakFilename, const TCHAR* InDestPakFil
 			}
 		}
 
-		FPakFile PakFile(&FPlatformFileManager::Get().GetPlatformFile(), *PakFilename, FParse::Param(FCommandLine::Get(), TEXT("signed")));
+		FPakFile PakFile(&FPlatformFileManager::Get().GetPlatformFile(), *PakFilename, bSigned);
 		if (PakFile.IsValid())
 		{
 			FArchive& PakReader = *PakFile.GetSharedReader(NULL);
@@ -2272,12 +2260,12 @@ public:
 };
 TMap<FName,FPackage*> FPackage::NameToPackageMap;
 
-bool ExportDependencies(const TCHAR * PakFilename, const TCHAR* GameName, const TCHAR* GameFolderName, const TCHAR* OutputFilenameBase )
+bool ExportDependencies(const TCHAR * PakFilename, const TCHAR* GameName, const TCHAR* GameFolderName, const TCHAR* OutputFilenameBase, bool bSigned)
 {
 	// Example command line used for this tool
 	// C:\Development\BB\WEX\Saved\StagedBuilds\WindowsNoEditor\WorldExplorers\Content\Paks\WorldExplorers-WindowsNoEditor.pak WorldExplorers WEX -exportdependencies=c:\dvtemp\output -debug -NoAssetRegistryCache -ForceDependsGathering
 	
-	FPakFile PakFile(&FPlatformFileManager::Get().GetPlatformFile(), PakFilename,FParse::Param(FCommandLine::Get(),TEXT("signed")));
+	FPakFile PakFile(&FPlatformFileManager::Get().GetPlatformFile(), PakFilename, bSigned);
 
 	if(PakFile.IsValid())
 	{
@@ -2436,7 +2424,7 @@ bool ExportDependencies(const TCHAR * PakFilename, const TCHAR* GameName, const
 	}
 }
 
-bool Repack(const FString& InputPakFile, const FString& OutputPakFile, const FPakCommandLineParameters& CmdLineParameters, const FKeyPair& SigningKey, const FAES::FAESKey& InEncryptionKey)
+bool Repack(const FString& InputPakFile, const FString& OutputPakFile, const FPakCommandLineParameters& CmdLineParameters, const FKeyPair& SigningKey, const FAES::FAESKey& InEncryptionKey, bool bSigned)
 {
 	bool bResult = false;
 
@@ -2445,7 +2433,7 @@ bool Repack(const FString& InputPakFile, const FString& OutputPakFile, const FPa
 	TArray<FPakInputPair> Entries;
 	TMap<FString, uint64> OrderMap;
 	FString TempDir = FPaths::EngineIntermediateDir() / TEXT("UnrealPak") / TEXT("Repack") / FPaths::GetBaseFilename(InputPakFile);
-	if (ExtractFilesFromPak(*InputPakFile, Hashes, *TempDir, false, InEncryptionKey, &Entries, &OrderMap))
+	if (ExtractFilesFromPak(*InputPakFile, Hashes, *TempDir, false, InEncryptionKey, bSigned, &Entries, &OrderMap))
 	{
 		TArray<FPakInputPair> FilesToAdd;
 		CollectFilesToAdd(FilesToAdd, Entries, OrderMap);
@@ -2497,253 +2485,305 @@ bool Repack(const FString& InputPakFile, const FString& OutputPakFile, const FPa
  * @param	ArgC	Command-line argument count
  * @param	ArgV	Argument strings
  */
-int32 ExecuteUnrealPak(int32 ArgC, TCHAR* ArgV[])
+bool ExecuteUnrealPak(const TCHAR* CmdLine)
 {
-	if (ArgC < 2)
-	{
-		UE_LOG(LogPakFile, Error, TEXT("No pak file name specified. Usage:"));
-		UE_LOG(LogPakFile, Error, TEXT("  UnrealPak <PakFilename> -Test"));
-		UE_LOG(LogPakFile, Error, TEXT("  UnrealPak <PakFilename> -List"));
-		UE_LOG(LogPakFile, Error, TEXT("  UnrealPak <PakFilename> <GameUProjectName> <GameFolderName> -ExportDependencies=<OutputFileBase> -NoAssetRegistryCache -ForceDependsGathering"));
-		UE_LOG(LogPakFile, Error, TEXT("  UnrealPak <PakFilename> -Extract <ExtractDir>"));
-		UE_LOG(LogPakFile, Error, TEXT("  UnrealPak <PakFilename> -Create=<ResponseFile> [Options]"));
-		UE_LOG(LogPakFile, Error, TEXT("  UnrealPak <PakFilename> -Dest=<MountPoint>"));
-		UE_LOG(LogPakFile, Error, TEXT("  UnrealPak <PakFilename> -Repack [-Output=Path] [Options]"));
-		UE_LOG(LogPakFile, Error, TEXT("  UnrealPak GenerateKeys=<KeyFilename>"));
-		UE_LOG(LogPakFile, Error, TEXT("  UnrealPak GeneratePrimeTable=<KeyFilename> [-TableMax=<N>]"));
-		UE_LOG(LogPakFile, Error, TEXT("  UnrealPak <PakFilename1> <PakFilename2> -diff"));
-		UE_LOG(LogPakFile, Error, TEXT("  UnrealPak -TestEncryption"));
-		UE_LOG(LogPakFile, Error, TEXT("  Options:"));
-		UE_LOG(LogPakFile, Error, TEXT("    -blocksize=<BlockSize>"));
-		UE_LOG(LogPakFile, Error, TEXT("    -bitwindow=<BitWindow>"));
-		UE_LOG(LogPakFile, Error, TEXT("    -compress"));
-		UE_LOG(LogPakFile, Error, TEXT("    -encrypt"));
-		UE_LOG(LogPakFile, Error, TEXT("    -order=<OrderingFile>"));
-		UE_LOG(LogPakFile, Error, TEXT("    -diff (requires 2 filenames first)"));
-		UE_LOG(LogPakFile, Error, TEXT("    -enginedir (specify engine dir for when using ini encryption configs)"));
-		UE_LOG(LogPakFile, Error, TEXT("    -projectdir (specify project dir for when using ini encryption configs)"));
-		UE_LOG(LogPakFile, Error, TEXT("    -encryptionini (specify ini base name to gather encryption settings from)"));
-		UE_LOG(LogPakFile, Error, TEXT("    -extracttomountpoint (Extract to mount point path of pak file)"));
-		UE_LOG(LogPakFile, Error, TEXT("    -encryptindex (encrypt the pak file index, making it unusable in unrealpak without supplying the key)"));
-		UE_LOG(LogPakFile, Error, TEXT("    -compressor=<DllPath> (register a custom compressor)"))
-		UE_LOG(LogPakFile, Error, TEXT("    -overrideplatformcompressor (override the native platform compressor)"))
-		return 1;
-	}
-
-	double StartTime = FPlatformTime::Seconds();
+	// Parse all the non-option arguments from the command line
+	TArray<FString> NonOptionArguments;
+	for (const TCHAR* CmdLineEnd = CmdLine; *CmdLineEnd != 0;)
+	{
+		FString Argument = FParse::Token(CmdLineEnd, false);
+		if (Argument.Len() > 0 && !Argument.StartsWith(TEXT("-")))
+		{
+			NonOptionArguments.Add(Argument);
+		}
+	}
+
 	FKeyPair SigningKey;
 	FAES::FAESKey EncryptionKey;
-	PrepareEncryptionAndSigningKeys(SigningKey, EncryptionKey);
+	PrepareEncryptionAndSigningKeys(CmdLine, SigningKey, EncryptionKey);
 
-	FPakCommandLineParameters CmdLineParameters;
-	int32 Result = 0;
 	FString KeyFilename;
-	if (FParse::Value(FCommandLine::Get(), TEXT("GenerateKeys="), KeyFilename, false))
+	if (FParse::Value(CmdLine, TEXT("GenerateKeys="), KeyFilename, false))
 	{
-		Result = GenerateKeys(*KeyFilename) ? 0 : 1;
+		return GenerateKeys(*KeyFilename);
 	}
-	else if (FParse::Value(FCommandLine::Get(), TEXT("GeneratePrimeTable="), KeyFilename, false))
+
+	if (FParse::Value(CmdLine, TEXT("GeneratePrimeTable="), KeyFilename, false))
 	{
 		int64 MaxPrimeValue = 10000;
-		FParse::Value(FCommandLine::Get(), TEXT("TableMax="), MaxPrimeValue);
+		FParse::Value(CmdLine, TEXT("TableMax="), MaxPrimeValue);
 		GeneratePrimeNumberTable(MaxPrimeValue, *KeyFilename);
+		return true;
 	}
-	else if (FParse::Param(FCommandLine::Get(), TEXT("TestEncryption")))
+
+	if (FParse::Param(CmdLine, TEXT("TestEncryption")))
 	{
 		void TestEncryption();
 		TestEncryption();
+		return true;
 	}
-	else 
+
+	if (FParse::Param(CmdLine, TEXT("Test")))
 	{
-		FString ExportDependencyFilename;
-		if (FParse::Param(FCommandLine::Get(), TEXT("Test")))
+		if (NonOptionArguments.Num() != 1)
 		{
-			FString PakFilename = GetPakPath(ArgV[1], false);
-			Result = TestPakFile(*PakFilename) ? 0 : 1;
+			UE_LOG(LogPakFile, Error, TEXT("Incorrect arguments. Expected: -Test <PakFile>"));
+			return false;
 		}
-		else if (FParse::Param(FCommandLine::Get(), TEXT("List")))
+
+		FString PakFilename = GetPakPath(*NonOptionArguments[0], false);
+		return TestPakFile(CmdLine, *PakFilename);
+	}
+
+	if (FParse::Param(CmdLine, TEXT("List")))
+	{
+		if (NonOptionArguments.Num() != 1)
 		{
-			int64 SizeFilter = 0;
-			FParse::Value(FCommandLine::Get(), TEXT("SizeFilter="), SizeFilter);
+			UE_LOG(LogPakFile, Error, TEXT("Incorrect arguments. Expected: -List <PakFile> [-SizeFilter=N] [-Signed]"));
+			return false;
+		}
+
+		int64 SizeFilter = 0;
+		FParse::Value(CmdLine, TEXT("SizeFilter="), SizeFilter);
+
+		FString PakFilename = GetPakPath(*NonOptionArguments[0], false);
+		bool bSigned = FParse::Param(CmdLine, TEXT("signed"));
+		return ListFilesInPak(*PakFilename, SizeFilter, bSigned);
+	}
 
-			FString PakFilename = GetPakPath(ArgV[1], false);
-			Result = ListFilesInPak(*PakFilename, SizeFilter) ? 0 : 1;
+	FString ExportDependencyFilename;
+	if(FParse::Value(CmdLine,TEXT("ExportDependencies="),ExportDependencyFilename,false))
+	{
+		if(NonOptionArguments.Num() != 3)
+		{
+			UE_LOG(LogPakFile,Error,TEXT("Incorrect arguments. Expected: -ExportDependencies=<FileName> <PakFile> <GameName> <GameFolderName> [-Signed]"));
+			return false;
 		}
-		else if(FParse::Value(FCommandLine::Get(),TEXT("ExportDependencies="),ExportDependencyFilename,false))
-		{			
-			if(ArgC < 4)
-			{
-				UE_LOG(LogPakFile,Error,TEXT("Insufficient arguments."));
-			}
-			else
-			{
-				FString PakFilename = GetPakPath(ArgV[1],false);
-				Result = ExportDependencies(*PakFilename,ArgV[2],ArgV[3],*ExportDependencyFilename);					
-			}
+
+		FString PakFilename = GetPakPath(*NonOptionArguments[0], false);
+		const TCHAR* GameName = *NonOptionArguments[1];
+		const TCHAR* GameFolderName = *NonOptionArguments[2];
+		bool bSigned = FParse::Param(CmdLine, TEXT("signed"));
+		return ExportDependencies(*PakFilename, GameName, GameFolderName, *ExportDependencyFilename, bSigned);
+	}
+
+	if (FParse::Param(CmdLine, TEXT("Diff")))
+	{
+		if(NonOptionArguments.Num() != 2)
+		{
+			UE_LOG(LogPakFile,Error,TEXT("Incorrect arguments. Expected: -Diff <PakFile1> <PakFile2> [-NoUniques] [-NoUniquesFile1] [-NoUniquesFile2]"));
+			return false;
 		}
-		else if (FParse::Param(FCommandLine::Get(), TEXT("Diff")))
+
+		FString PakFilename1 = GetPakPath(*NonOptionArguments[0], false);
+		FString PakFilename2 = GetPakPath(*NonOptionArguments[1], false);
+
+		// Allow the suppression of unique file logging for one or both files
+		const bool bLogUniques = !FParse::Param(CmdLine, TEXT("nouniques"));
+		const bool bLogUniques1 = bLogUniques && !FParse::Param(CmdLine, TEXT("nouniquesfile1"));
+		const bool bLogUniques2 = bLogUniques && !FParse::Param(CmdLine, TEXT("nouniquesfile2"));
+
+		const bool bSigned = FParse::Param(CmdLine, TEXT("signed"));
+
+		return DiffFilesInPaks(*PakFilename1, *PakFilename2, bLogUniques1, bLogUniques2, EncryptionKey, bSigned);
+	}
+
+	if (FParse::Param(CmdLine, TEXT("Extract")))
+	{
+		if (NonOptionArguments.Num() != 2)
 		{
-			FString PakFilename1 = GetPakPath(ArgV[1], false);
-			FString PakFilename2 = GetPakPath(ArgV[2], false);
-			Result = DiffFilesInPaks(*PakFilename1, *PakFilename2, EncryptionKey) ? 0 : 1;
+			UE_LOG(LogPakFile, Error, TEXT("Incorrect arguments. Expected: -Extract <PakFile> <OutputPath>"));
+			return false;
 		}
-		else if (FParse::Param(FCommandLine::Get(), TEXT("Extract")))
+
+		FString PakFilename = GetPakPath(*NonOptionArguments[0], false);
+		bool bSigned = FParse::Param(CmdLine, TEXT("signed"));
+
+		FString DestPath = NonOptionArguments[1];
+		bool bExtractToMountPoint = FParse::Param(CmdLine, TEXT("ExtractToMountPoint"));
+		TMap<FString, FFileInfo> EmptyMap;
+		return ExtractFilesFromPak(*PakFilename, EmptyMap, *DestPath, bExtractToMountPoint, EncryptionKey, bSigned);
+	}
+
+	if (FParse::Param(CmdLine, TEXT("Repack")))
+	{
+		if (NonOptionArguments.Num() != 1)
 		{
-			FString PakFilename = GetPakPath(ArgV[1], false);
-			if (ArgC < 4)
+			UE_LOG(LogPakFile, Error, TEXT("Incorrect arguments. Expected: -Repack <PakFile> [-Output=<PakFile>] [-Signed]"));
+			return false;
+		}
+
+		TArray<FPakInputPair> Entries;
+		FPakCommandLineParameters CmdLineParameters;
+		ProcessCommandLine(CmdLine, NonOptionArguments, Entries, CmdLineParameters);
+
+		// Find all the input pak files
+		FString InputDir = FPaths::GetPath(*NonOptionArguments[0]);
+
+		TArray<FString> InputPakFiles;
+		IFileManager::Get().FindFiles(InputPakFiles, *InputDir, *FPaths::GetCleanFilename(*NonOptionArguments[0]));
+
+		for (int Idx = 0; Idx < InputPakFiles.Num(); Idx++)
+		{
+			InputPakFiles[Idx] = InputDir / InputPakFiles[Idx];
+		}
+
+		if (InputPakFiles.Num() == 0)
+		{
+			UE_LOG(LogPakFile, Error, TEXT("No files found matching '%s'"), *NonOptionArguments[0]);
+			return false;
+		}
+
+		// Find all the output paths
+		TArray<FString> OutputPakFiles;
+
+		FString OutputPath;
+		if (!FParse::Value(CmdLine, TEXT("Output="), OutputPath, false))
+		{
+			for (const FString& InputPakFile : InputPakFiles)
 			{
-				UE_LOG(LogPakFile, Error, TEXT("No extraction path specified."));
-				Result = 1;
+				OutputPakFiles.Add(InputPakFile);
 			}
-			else
+		}
+		else if (IFileManager::Get().DirectoryExists(*OutputPath))
+		{
+			for (const FString& InputPakFile : InputPakFiles)
 			{
-				FString DestPath = (ArgV[2][0] == '-') ? ArgV[3] : ArgV[2];
-				bool bExtractToMountPoint = FParse::Param(FCommandLine::Get(), TEXT("ExtractToMountPoint"));
-				TMap<FString, FFileInfo> EmptyMap;
-				Result = ExtractFilesFromPak(*PakFilename, EmptyMap, *DestPath, bExtractToMountPoint, EncryptionKey) ? 0 : 1;
+				OutputPakFiles.Add(FPaths::Combine(OutputPath, FPaths::GetCleanFilename(InputPakFile)));
 			}
 		}
-		else if (FParse::Param(FCommandLine::Get(), TEXT("Repack")))
+		else
 		{
-			TArray<FPakInputPair> Entries;
-			ProcessCommandLine(ArgC, ArgV, Entries, CmdLineParameters);
-
-			// Find all the input pak files
-			FString InputDir = FPaths::GetPath(ArgV[1]);
-
-			TArray<FString> InputPakFiles;
-			IFileManager::Get().FindFiles(InputPakFiles, *InputDir, *FPaths::GetCleanFilename(ArgV[1]));
-
-			for (int Idx = 0; Idx < InputPakFiles.Num(); Idx++)
+			for (const FString& InputPakFile : InputPakFiles)
 			{
-				InputPakFiles[Idx] = InputDir / InputPakFiles[Idx];
+				OutputPakFiles.Add(OutputPath);
 			}
+		}
 
-			if (InputPakFiles.Num() == 0)
+		// Repack them all
+		bool bSigned = FParse::Param(CmdLine, TEXT("signed"));
+		for (int Idx = 0; Idx < InputPakFiles.Num(); Idx++)
+		{
+			UE_LOG(LogPakFile, Display, TEXT("Repacking %s into %s"), *InputPakFiles[Idx], *OutputPakFiles[Idx]);
+			if (!Repack(InputPakFiles[Idx], OutputPakFiles[Idx], CmdLineParameters, SigningKey, EncryptionKey, bSigned))
 			{
-				UE_LOG(LogPakFile, Error, TEXT("No files found matching '%s'"), ArgV[1]);
-				Result = 1;
+				return false;
 			}
-			else
-			{
-				// Find all the output paths
-				TArray<FString> OutputPakFiles;
+		}
 
-				FString OutputPath;
-				if (!FParse::Value(FCommandLine::Get(), TEXT("Output="), OutputPath, false))
-				{
-					for (const FString& InputPakFile : InputPakFiles)
-					{
-						OutputPakFiles.Add(InputPakFile);
-					}
-				}
-				else if (IFileManager::Get().DirectoryExists(*OutputPath))
-				{
-					for (const FString& InputPakFile : InputPakFiles)
-					{
-						OutputPakFiles.Add(FPaths::Combine(OutputPath, FPaths::GetCleanFilename(InputPakFile)));
-					}
-				}
-				else
-				{
-					for (const FString& InputPakFile : InputPakFiles)
-					{
-						OutputPakFiles.Add(OutputPath);
-					}
-				}
+		return true;
+	}
 
-				// Repack them all
-				for (int Idx = 0; Idx < InputPakFiles.Num(); Idx++)
-				{
-					UE_LOG(LogPakFile, Display, TEXT("Repacking %s into %s"), *InputPakFiles[Idx], *OutputPakFiles[Idx]);
-					if (!Repack(InputPakFiles[Idx], OutputPakFiles[Idx], CmdLineParameters, SigningKey, EncryptionKey))
-					{
-						Result = 1;
-						break;
-					}
-				}
-			}
+	if(NonOptionArguments.Num() > 0)
+	{
+		// since this is for creation, we pass true to make it not look in LaunchDir
+		FString PakFilename = GetPakPath(*NonOptionArguments[0], true);
+		bool bSigned = FParse::Param(CmdLine, TEXT("signed"));
+
+		// List of all items to add to pak file
+		TArray<FPakInputPair> Entries;
+		FPakCommandLineParameters CmdLineParameters;
+		ProcessCommandLine(CmdLine, NonOptionArguments, Entries, CmdLineParameters);
+
+		TMap<FString, uint64> OrderMap;
+		FString ResponseFile;
+		if (FParse::Value(CmdLine, TEXT("-order="), ResponseFile) && !ProcessOrderFile(*ResponseFile, OrderMap))
+		{
+			return false;
 		}
-		else
+
+		if (Entries.Num() == 0)
 		{
-			// since this is for creation, we pass true to make it not look in LaunchDir
-			FString PakFilename = GetPakPath(ArgV[1], true);
+			UE_LOG(LogPakFile, Error, TEXT("No files specified to add to pak file."));
+			return false;
+		}
 
-			// List of all items to add to pak file
-			TArray<FPakInputPair> Entries;
-			ProcessCommandLine(ArgC, ArgV, Entries, CmdLineParameters);
-			TMap<FString, uint64> OrderMap;
-			ProcessOrderFile(ArgC, ArgV, OrderMap);
+		TMap<FString, FFileInfo> SourceFileHashes;
 
-			if (Entries.Num() == 0)
+		if ( CmdLineParameters.GeneratePatch )
+		{
+			FString OutputPath;
+			if (!FParse::Value(CmdLine, TEXT("TempFiles="), OutputPath))
 			{
-				UE_LOG(LogPakFile, Error, TEXT("No files specified to add to pak file."));
-				Result = 1;
+				OutputPath = FPaths::GetPath(PakFilename) / FString(TEXT("TempFiles"));
 			}
-			else
-			{
-				TMap<FString, FFileInfo> SourceFileHashes;
-
-				if ( CmdLineParameters.GeneratePatch )
-				{
-					FString OutputPath;
-					if (!FParse::Value(FCommandLine::Get(), TEXT("TempFiles="), OutputPath))
-					{
-						OutputPath = FPaths::GetPath(PakFilename) / FString(TEXT("TempFiles"));
-					}
 
-					IFileManager::Get().DeleteDirectory(*OutputPath);
+			IFileManager::Get().DeleteDirectory(*OutputPath);
 
-					// Check command line for the "patchcryptokeys" param, which will tell us where to look for the encryption keys that
-					// we need to access the patch reference data
-					FString PatchReferenceCryptoKeysFilename;
-					FAES::FAESKey PatchReferenceEncryptionKey = EncryptionKey;
-					if (FParse::Value(FCommandLine::Get(), TEXT("PatchCryptoKeys="), PatchReferenceCryptoKeysFilename))
-					{
-						FKeyPair UnusedSigningKey;
-						PrepareEncryptionAndSigningKeysFromCryptoKeyCache(PatchReferenceCryptoKeysFilename, UnusedSigningKey, PatchReferenceEncryptionKey);
-					}
+			// Check command line for the "patchcryptokeys" param, which will tell us where to look for the encryption keys that
+			// we need to access the patch reference data
+			FString PatchReferenceCryptoKeysFilename;
+			FAES::FAESKey PatchReferenceEncryptionKey = EncryptionKey;
+			if (FParse::Value(CmdLine, TEXT("PatchCryptoKeys="), PatchReferenceCryptoKeysFilename))
+			{
+				FKeyPair UnusedSigningKey;
+				PrepareEncryptionAndSigningKeysFromCryptoKeyCache(PatchReferenceCryptoKeysFilename, UnusedSigningKey, PatchReferenceEncryptionKey);
+			}
 
-					UE_LOG(LogPakFile, Display, TEXT("Generating patch from %s."), *CmdLineParameters.SourcePatchPakFilename, true );
+			UE_LOG(LogPakFile, Display, TEXT("Generating patch from %s."), *CmdLineParameters.SourcePatchPakFilename, true );
 
-					if (!GenerateHashesFromPak(*CmdLineParameters.SourcePatchPakFilename, *PakFilename, SourceFileHashes, true, PatchReferenceEncryptionKey))
-					{
-						if (ExtractFilesFromPak(*CmdLineParameters.SourcePatchPakFilename, SourceFileHashes, *OutputPath, true, PatchReferenceEncryptionKey) == false)
-						{
-							UE_LOG(LogPakFile, Warning, TEXT("Unable to extract files from source pak file for patch"));
-						}
-						else
-						{
-							CmdLineParameters.SourcePatchDiffDirectory = OutputPath;
-						}
-					}
+			if (!GenerateHashesFromPak(*CmdLineParameters.SourcePatchPakFilename, *PakFilename, SourceFileHashes, true, PatchReferenceEncryptionKey, bSigned))
+			{
+				if (ExtractFilesFromPak(*CmdLineParameters.SourcePatchPakFilename, SourceFileHashes, *OutputPath, true, PatchReferenceEncryptionKey, bSigned) == false)
+				{
+					UE_LOG(LogPakFile, Warning, TEXT("Unable to extract files from source pak file for patch"));
 				}
+				else
+				{
+					CmdLineParameters.SourcePatchDiffDirectory = OutputPath;
+				}
+			}
+		}
 
 
-				// Start collecting files
-				TArray<FPakInputPair> FilesToAdd;
-				CollectFilesToAdd(FilesToAdd, Entries, OrderMap);
+		// Start collecting files
+		TArray<FPakInputPair> FilesToAdd;
+		CollectFilesToAdd(FilesToAdd, Entries, OrderMap);
 
-				if ( CmdLineParameters.GeneratePatch )
-				{
-					// if we are generating a patch here we remove files which are already shipped...
-					RemoveIdenticalFiles(FilesToAdd, CmdLineParameters.SourcePatchDiffDirectory, SourceFileHashes);
-				}
+		if ( CmdLineParameters.GeneratePatch )
+		{
+			// if we are generating a patch here we remove files which are already shipped...
+			RemoveIdenticalFiles(FilesToAdd, CmdLineParameters.SourcePatchDiffDirectory, SourceFileHashes);
+		}
 
 
-				Result = CreatePakFile(*PakFilename, FilesToAdd, CmdLineParameters, SigningKey, EncryptionKey) ? 0 : 1;
+		bool bResult = CreatePakFile(*PakFilename, FilesToAdd, CmdLineParameters, SigningKey, EncryptionKey);
 
-				if (CmdLineParameters.GeneratePatch)
-				{
-					FString OutputPath = FPaths::GetPath(PakFilename) / FString(TEXT("TempFiles"));
-					// delete the temporary directory
-					IFileManager::Get().DeleteDirectory(*OutputPath, false, true);
-				}
-			}
+		if (CmdLineParameters.GeneratePatch)
+		{
+			FString OutputPath = FPaths::GetPath(PakFilename) / FString(TEXT("TempFiles"));
+			// delete the temporary directory
+			IFileManager::Get().DeleteDirectory(*OutputPath, false, true);
 		}
-	}
 
-	UE_LOG(LogPakFile, Display, TEXT("Unreal pak executed in %f seconds"), FPlatformTime::Seconds() - StartTime );
+		return bResult;
+	}
 
-	return Result;
+	UE_LOG(LogPakFile, Error, TEXT("No pak file name specified. Usage:"));
+	UE_LOG(LogPakFile, Error, TEXT("  UnrealPak <PakFilename> -Test"));
+	UE_LOG(LogPakFile, Error, TEXT("  UnrealPak <PakFilename> -List"));
+	UE_LOG(LogPakFile, Error, TEXT("  UnrealPak <PakFilename> <GameUProjectName> <GameFolderName> -ExportDependencies=<OutputFileBase> -NoAssetRegistryCache -ForceDependsGathering"));
+	UE_LOG(LogPakFile, Error, TEXT("  UnrealPak <PakFilename> -Extract <ExtractDir>"));
+	UE_LOG(LogPakFile, Error, TEXT("  UnrealPak <PakFilename> -Create=<ResponseFile> [Options]"));
+	UE_LOG(LogPakFile, Error, TEXT("  UnrealPak <PakFilename> -Dest=<MountPoint>"));
+	UE_LOG(LogPakFile, Error, TEXT("  UnrealPak <PakFilename> -Repack [-Output=Path] [Options]"));
+	UE_LOG(LogPakFile, Error, TEXT("  UnrealPak GenerateKeys=<KeyFilename>"));
+	UE_LOG(LogPakFile, Error, TEXT("  UnrealPak GeneratePrimeTable=<KeyFilename> [-TableMax=<N>]"));
+	UE_LOG(LogPakFile, Error, TEXT("  UnrealPak <PakFilename1> <PakFilename2> -diff"));
+	UE_LOG(LogPakFile, Error, TEXT("  UnrealPak -TestEncryption"));
+	UE_LOG(LogPakFile, Error, TEXT("  Options:"));
+	UE_LOG(LogPakFile, Error, TEXT("    -blocksize=<BlockSize>"));
+	UE_LOG(LogPakFile, Error, TEXT("    -bitwindow=<BitWindow>"));
+	UE_LOG(LogPakFile, Error, TEXT("    -compress"));
+	UE_LOG(LogPakFile, Error, TEXT("    -encrypt"));
+	UE_LOG(LogPakFile, Error, TEXT("    -order=<OrderingFile>"));
+	UE_LOG(LogPakFile, Error, TEXT("    -diff (requires 2 filenames first)"));
+	UE_LOG(LogPakFile, Error, TEXT("    -enginedir (specify engine dir for when using ini encryption configs)"));
+	UE_LOG(LogPakFile, Error, TEXT("    -projectdir (specify project dir for when using ini encryption configs)"));
+	UE_LOG(LogPakFile, Error, TEXT("    -encryptionini (specify ini base name to gather encryption settings from)"));
+	UE_LOG(LogPakFile, Error, TEXT("    -extracttomountpoint (Extract to mount point path of pak file)"));
+	UE_LOG(LogPakFile, Error, TEXT("    -encryptindex (encrypt the pak file index, making it unusable in unrealpak without supplying the key)"));
+	UE_LOG(LogPakFile, Error, TEXT("    -compressor=<DllPath> (register a custom compressor)"))
+	UE_LOG(LogPakFile, Error, TEXT("    -overrideplatformcompressor (override the native platform compressor)"))
+	return false;
 }
diff --git a/Engine/Source/Developer/PakFileUtilities/Public/PakFileUtilities.h b/Engine/Source/Developer/PakFileUtilities/Public/PakFileUtilities.h
--- a/Engine/Source/Developer/PakFileUtilities/Public/PakFileUtilities.h
+++ b/Engine/Source/Developer/PakFileUtilities/Public/PakFileUtilities.h
@@ -2,5 +2,5 @@
 
 #include "CoreMinimal.h"
 
-PAKFILEUTILITIES_API int32 ExecuteUnrealPak(int ArgC, TCHAR* ArgV[]);
+PAKFILEUTILITIES_API bool ExecuteUnrealPak(const TCHAR* CmdLine);
 
diff --git a/Engine/Source/Editor/UnrealEd/Classes/Commandlets/UnrealPakCommandlet.h b/Engine/Source/Editor/UnrealEd/Classes/Commandlets/UnrealPakCommandlet.h
new file mode 100644
--- /dev/null
+++ b/Engine/Source/Editor/UnrealEd/Classes/Commandlets/UnrealPakCommandlet.h
@@ -0,0 +1,19 @@
+// Copyright 1998-2018 Epic Games, Inc. All Rights Reserved.
+
+#pragma once
+
+#include "CoreMinimal.h"
+#include "UObject/ObjectMacros.h"
+#include "Commandlets/Commandlet.h"
+#include "UnrealPakCommandlet.generated.h"
+
+UCLASS()
+class UUnrealPakCommandlet : public UCommandlet
+{
+	GENERATED_UCLASS_BODY()
+	//~ Begin UCommandlet Interface
+	virtual int32 Main(const FString& Params) override;
+	//~ End UCommandlet Interface
+};
+
+
diff --git a/Engine/Source/Editor/UnrealEd/Private/Commandlets/UnrealPakCommandlet.cpp b/Engine/Source/Editor/UnrealEd/Private/Commandlets/UnrealPakCommandlet.cpp
new file mode 100644
--- /dev/null
+++ b/Engine/Source/Editor/UnrealEd/Private/Commandlets/UnrealPakCommandlet.cpp
@@ -0,0 +1,14 @@
+// Copyright 1998-2018 Epic Games, Inc. All Rights Reserved.
+
+#include "Commandlets/UnrealPakCommandlet.h"
+#include "PakFileUtilities.h"
+
+UUnrealPakCommandlet::UUnrealPakCommandlet(const FObjectInitializer& ObjectInitializer)
+	: Super(ObjectInitializer)
+{
+}
+
+int32 UUnrealPakCommandlet::Main(const FString& Params)
+{
+	return ExecuteUnrealPak(*Params)? 0 : 1;
+}
diff --git a/Engine/Source/Editor/UnrealEd/UnrealEd.Build.cs b/Engine/Source/Editor/UnrealEd/UnrealEd.Build.cs
--- a/Engine/Source/Editor/UnrealEd/UnrealEd.Build.cs
+++ b/Engine/Source/Editor/UnrealEd/UnrealEd.Build.cs
@@ -161,6 +161,7 @@ public class UnrealEd : ModuleRules
 				"ClothingSystemRuntime",
 				"ClothingSystemRuntimeInterface",
 				"PIEPreviewDeviceProfileSelector",
+				"PakFileUtilities",
 				"TimeManagement",
 			}
 		);
diff --git a/Engine/Source/Programs/UnrealPak/Private/UnrealPak.cpp b/Engine/Source/Programs/UnrealPak/Private/UnrealPak.cpp
--- a/Engine/Source/Programs/UnrealPak/Private/UnrealPak.cpp
+++ b/Engine/Source/Programs/UnrealPak/Private/UnrealPak.cpp
@@ -3,6 +3,7 @@
 #include "UnrealPak.h"
 #include "RequiredProgramMainCPPInclude.h"
 #include "PakFileUtilities.h"
+#include "IPlatformFilePak.h"
 
 IMPLEMENT_APPLICATION(UnrealPak, "UnrealPak");
 
@@ -11,7 +12,11 @@ INT32_MAIN_INT32_ARGC_TCHAR_ARGV()
 	// start up the main loop
 	GEngineLoop.PreInit(ArgC, ArgV);
 
-	int32 Result = ExecuteUnrealPak(ArgC, ArgV);
+	double StartTime = FPlatformTime::Seconds();
+
+	int32 Result = ExecuteUnrealPak(FCommandLine::Get());
+
+	UE_LOG(LogPakFile, Display, TEXT("Unreal pak executed in %f seconds"), FPlatformTime::Seconds() - StartTime );
 
 	GLog->Flush();
 
