Fixing compile error in PakFileUtilities

diff --git a/Engine/Source/Developer/PakFileUtilities/Private/PakFileUtilities.cpp b/Engine/Source/Developer/PakFileUtilities/Private/PakFileUtilities.cpp
--- a/Engine/Source/Developer/PakFileUtilities/Private/PakFileUtilities.cpp
+++ b/Engine/Source/Developer/PakFileUtilities/Private/PakFileUtilities.cpp
@@ -2532,7 +2532,8 @@ bool ExecuteUnrealPak(const TCHAR* CmdLine)
 		}
 
 		FString PakFilename = GetPakPath(*NonOptionArguments[0], false);
-		return TestPakFile(CmdLine, *PakFilename);
+		bool bSigned = FParse::Param(CmdLine, TEXT("signed"));
+		return TestPakFile(*PakFilename, bSigned);
 	}
 
 	if (FParse::Param(CmdLine, TEXT("List")))
