Fixed an erroneous usage of TUniquePtr<uint8>as a pointer to a uint8 array when creating pak files. Caused a crash when compression was enabled, and has probably surfaced because pak generation is now done by an editor commandlet rather than a standalone program.

diff --git a/Engine/Source/Developer/PakFileUtilities/Private/PakFileUtilities.cpp b/Engine/Source/Developer/PakFileUtilities/Private/PakFileUtilities.cpp
--- a/Engine/Source/Developer/PakFileUtilities/Private/PakFileUtilities.cpp
+++ b/Engine/Source/Developer/PakFileUtilities/Private/PakFileUtilities.cpp
@@ -111,22 +111,22 @@ struct FCompressedFileBuffer
 	{
 		if(RequiredSpace > CompressedBufferSize)
 		{
-			uint8* NewPtr = (uint8*)FMemory::Malloc(RequiredSpace);
-			FMemory::Memcpy(NewPtr,CompressedBuffer.Get(),CompressedBufferSize);
-			CompressedBuffer.Reset(NewPtr);
+			TUniquePtr<uint8[]> NewCompressedBuffer = MakeUnique<uint8[]>(RequiredSpace);
+			FMemory::Memcpy(NewCompressedBuffer.Get(), CompressedBuffer.Get(), CompressedBufferSize);
+			CompressedBuffer = MoveTemp(NewCompressedBuffer);
 			CompressedBufferSize = RequiredSpace;
 		}
 	}
 
 	bool CompressFileToWorkingBuffer(const FPakInputPair& InFile,uint8*& InOutPersistentBuffer,int64& InOutBufferSize,ECompressionFlags CompressionMethod,const int32 CompressionBlockSize,const int32 CompressionBitWindow);
 
-	int64						 OriginalSize;
-	int64						 TotalCompressedSize;
-	int32						 FileCompressionBlockSize;
-	ECompressionFlags			 FileCompressionMethod;
-	TArray<FPakCompressedBlock>  CompressedBlocks;
-	int64						 CompressedBufferSize;
-	TUniquePtr<uint8>		     CompressedBuffer;
+	int64				OriginalSize;
+	int64				TotalCompressedSize;
+	int32				FileCompressionBlockSize;
+	ECompressionFlags		FileCompressionMethod;
+	TArray<FPakCompressedBlock>	CompressedBlocks;
+	int64				CompressedBufferSize;
+	TUniquePtr<uint8[]>		CompressedBuffer;
 };
 
 FString GetLongestPath(TArray<FPakInputPair>& FilesToAdd)
