Pak creation now uses the bEnablePakSigning setting from the crypto config json file

diff --git a/Engine/Source/Developer/PakFileUtilities/Private/PakFileUtilities.cpp b/Engine/Source/Developer/PakFileUtilities/Private/PakFileUtilities.cpp
--- a/Engine/Source/Developer/PakFileUtilities/Private/PakFileUtilities.cpp
+++ b/Engine/Source/Developer/PakFileUtilities/Private/PakFileUtilities.cpp
@@ -855,16 +855,20 @@ void PrepareEncryptionAndSigningKeysFromCryptoKeyCache(const FString& InFilename
 				}
 			}
 
-			const TSharedPtr<FJsonObject>* SigningKey;
-			if (RootObject->TryGetObjectField(TEXT("SigningKey"), SigningKey))
-			{
-				TSharedPtr<FJsonObject> PublicKey = (*SigningKey)->GetObjectField(TEXT("PublicKey"));
-				TSharedPtr<FJsonObject> PrivateKey = (*SigningKey)->GetObjectField(TEXT("PrivateKey"));
-				OutSigningKey.PublicKey.Exponent = ParseEncryptionIntFromJson(PublicKey, TEXT("Exponent"));
-				OutSigningKey.PublicKey.Modulus = ParseEncryptionIntFromJson(PublicKey, TEXT("Modulus"));
-				OutSigningKey.PrivateKey.Exponent = ParseEncryptionIntFromJson(PrivateKey, TEXT("Exponent"));
-				OutSigningKey.PrivateKey.Modulus = ParseEncryptionIntFromJson(PrivateKey, TEXT("Modulus"));
-				check(OutSigningKey.PublicKey.Modulus == OutSigningKey.PrivateKey.Modulus);
+			bool bEnablePakSigning = false;
+			if (RootObject->TryGetBoolField(TEXT("bEnablePakSigning"), bEnablePakSigning))
+			{
+				const TSharedPtr<FJsonObject>* SigningKey = nullptr;
+				if (bEnablePakSigning && RootObject->TryGetObjectField(TEXT("SigningKey"), SigningKey))
+				{
+					TSharedPtr<FJsonObject> PublicKey = (*SigningKey)->GetObjectField(TEXT("PublicKey"));
+					TSharedPtr<FJsonObject> PrivateKey = (*SigningKey)->GetObjectField(TEXT("PrivateKey"));
+					OutSigningKey.PublicKey.Exponent = ParseEncryptionIntFromJson(PublicKey, TEXT("Exponent"));
+					OutSigningKey.PublicKey.Modulus = ParseEncryptionIntFromJson(PublicKey, TEXT("Modulus"));
+					OutSigningKey.PrivateKey.Exponent = ParseEncryptionIntFromJson(PrivateKey, TEXT("Exponent"));
+					OutSigningKey.PrivateKey.Modulus = ParseEncryptionIntFromJson(PrivateKey, TEXT("Modulus"));
+					check(OutSigningKey.PublicKey.Modulus == OutSigningKey.PrivateKey.Modulus);
+				}
 			}
 		}
 	}
