UnrealPak: Move "ExportDependencies" command into an editor commandlet, since it relies on the UObject system, asset registry, etc...

diff --git a/Engine/Source/Developer/PakFileUtilities/PakFileUtilities.Build.cs b/Engine/Source/Developer/PakFileUtilities/PakFileUtilities.Build.cs
--- a/Engine/Source/Developer/PakFileUtilities/PakFileUtilities.Build.cs
+++ b/Engine/Source/Developer/PakFileUtilities/PakFileUtilities.Build.cs
@@ -6,17 +6,11 @@ public class PakFileUtilities : ModuleRules
 {
 	public PakFileUtilities(ReadOnlyTargetRules Target) : base(Target)
 	{
-		PrivateDependencyModuleNames.AddRange(new string[] { "Core", "CoreUObject", "PakFile", "Json", "Projects", "ApplicationCore" });
+		PrivateDependencyModuleNames.AddRange(new string[] { "Core", "PakFile", "Json", "Projects", "ApplicationCore" });
 
         PrivateIncludePathModuleNames.AddRange(
             new string[] {
-                "AssetRegistry",
                 "Json"
         });
-
-        DynamicallyLoadedModuleNames.AddRange(
-            new string[] {
-                "AssetRegistry"
-        });
     }
 }
diff --git a/Engine/Source/Developer/PakFileUtilities/Private/PakFileUtilities.cpp b/Engine/Source/Developer/PakFileUtilities/Private/PakFileUtilities.cpp
--- a/Engine/Source/Developer/PakFileUtilities/Private/PakFileUtilities.cpp
+++ b/Engine/Source/Developer/PakFileUtilities/Private/PakFileUtilities.cpp
@@ -9,7 +9,6 @@
 #include "Misc/AES.h"
 #include "Templates/UniquePtr.h"
 #include "Serialization/LargeMemoryWriter.h"
-#include "AssetRegistryModule.h"
 #include "ProfilingDebugging/DiagnosticTable.h"
 #include "Serialization/JsonSerializer.h"
 #include "Misc/Base64.h"
@@ -2175,271 +2174,6 @@ FString GetPakPath(const TCHAR* SpecifiedPath, bool bIsForCreation)
 	return PakFilename;
 }
 
-struct FPackage
-{
-public:
-	TSet<FPackage*> DirectlyReferencing;
-	TSet<FPackage*> DirectlyReferencedBy;
-	TSet<FPackage*> AllReferences;
-
-	FName Name;
-	int64 InclusiveSize;
-	int64 ExclusiveSize;
-	int Id;
-
-private:
-	bool bUpdateHelper;
-	static TMap<FName,FPackage*> NameToPackageMap;
-
-public:
-	FPackage(const FName& InName,int InId)
-		: Name(InName)
-		, InclusiveSize(0)
-		, ExclusiveSize(0)
-		, Id(InId)
-		, bUpdateHelper(false)
-	{}
-
-	static FPackage* FindOrCreatePackage( FName PackageName )
-	{
-		static int Id = 1;
-		FPackage* Package = NameToPackageMap.FindRef(PackageName);
-		if(!Package)
-		{
-			Package = new FPackage(PackageName,Id++);
-			NameToPackageMap.Add(PackageName,Package);
-		}
-		return Package;
-	}
-
-	void ResetUpdateHelper()
-	{
-		bUpdateHelper = false;
-	}
-
-	void RecurseUpdateReferences()
-	{
-		if( !bUpdateHelper )
-		{
-			bUpdateHelper = true;
-			for( auto& DirectReference : DirectlyReferencing )
-			{
-				AllReferences.Add(DirectReference);
-				DirectReference->RecurseUpdateReferences();
-				AllReferences.Append(DirectReference->AllReferences);
-			}
-		}
-	}
-
-	void UpdateInclusiveSize()
-	{
-		InclusiveSize = ExclusiveSize;
-		for(auto& Reference : AllReferences)
-		{
-			InclusiveSize += Reference->ExclusiveSize;
-		}
-	}
-
-	static void GetAllPackages( TArray<FPackage*>& OutPackages )
-	{
-		OutPackages.Reset(NameToPackageMap.Num());
-		for( const auto& Entry : NameToPackageMap )
-		{
-			OutPackages.Add(Entry.Value);
-		}
-	}
-
-	TArray< TSharedPtr<FJsonValue> > ToJsonHelper( const TSet<FPackage*>& Packages )
-	{
-		TArray< TSharedPtr<FJsonValue> > JsonPackageNames;
-		for( const auto Package : Packages )
-		{
-			JsonPackageNames.Add(MakeShareable(new FJsonValueString(Package->Name.ToString())));
-		}
-		return JsonPackageNames;
-	}
-
-	TSharedPtr<FJsonObject> ToJsonObject()
-	{
-		TSharedPtr<FJsonObject> JsonPackageObject = MakeShareable(new FJsonObject);
-		
-		JsonPackageObject->SetStringField(TEXT("Name"),*Name.ToString());
-		JsonPackageObject->SetNumberField(TEXT("InclusiveSize"),InclusiveSize);
-		JsonPackageObject->SetNumberField(TEXT("ExclusiveSize"),ExclusiveSize);
-
-		JsonPackageObject->SetArrayField(TEXT("DirectlyReferencing"),ToJsonHelper(DirectlyReferencing));
-		JsonPackageObject->SetArrayField(TEXT("DirectlyReferencedBy"),ToJsonHelper(DirectlyReferencedBy));
-		JsonPackageObject->SetArrayField(TEXT("AllReferences"),ToJsonHelper(AllReferences));
-
-		return JsonPackageObject;
-	}
-};
-TMap<FName,FPackage*> FPackage::NameToPackageMap;
-
-bool ExportDependencies(const TCHAR * PakFilename, const TCHAR* GameName, const TCHAR* GameFolderName, const TCHAR* OutputFilenameBase, bool bSigned)
-{
-	// Example command line used for this tool
-	// C:\Development\BB\WEX\Saved\StagedBuilds\WindowsNoEditor\WorldExplorers\Content\Paks\WorldExplorers-WindowsNoEditor.pak WorldExplorers WEX -exportdependencies=c:\dvtemp\output -debug -NoAssetRegistryCache -ForceDependsGathering
-	
-	FPakFile PakFile(&FPlatformFileManager::Get().GetPlatformFile(), PakFilename, bSigned);
-
-	if(PakFile.IsValid())
-	{
-		// Get size information from PAK file.
-		{
-			TArray<FPakFile::FFileIterator> Records;
-			FString PakGameContentFolder = FString(GameName) + TEXT("/Content");
-			for(FPakFile::FFileIterator It(PakFile); It; ++It)
-			{
-				FString PackageName;
-				It.Filename().Split(TEXT("."),&PackageName,NULL);
-				int64 Size = It.Info().Size;			
-
-				if( PackageName.StartsWith(TEXT("Engine/Content")) )
-				{
-					PackageName = PackageName.Replace(TEXT("Engine/Content"),TEXT("/Engine"));
-				}
-				else if( PackageName.StartsWith(*PakGameContentFolder))
-				{
-					PackageName = PackageName.Replace(*PakGameContentFolder,TEXT("/Game"));
-				}
-
-				FPackage* Package = FPackage::FindOrCreatePackage(FName(*PackageName));
-				Package->ExclusiveSize += Size;
-			}
-		}
-
-		TMap<FName,FName> PackageToClassMap;
-
-		// Combine with dependency information from asset registry.
-		{
-			FString ProjectDir = FString(TEXT("../../../")) + GameFolderName + TEXT("/");
-			FPlatformMisc::SetOverrideProjectDir(*ProjectDir);
-			FAssetRegistryModule& AssetRegistryModule = FModuleManager::Get().LoadModuleChecked<FAssetRegistryModule>(TEXT("AssetRegistry"));
-		
-			AssetRegistryModule.Get().SearchAllAssets(true);
-			TArray<FAssetData> AssetData;		
-			AssetRegistryModule.Get().GetAllAssets(AssetData,true);	
-	
-			TSet<FName> PackageNames;
-			for( int i=0; i<AssetData.Num(); i++ )
-			{
-				PackageNames.Add(AssetData[i].PackageName);
-				PackageToClassMap.Add(AssetData[i].PackageName,AssetData[i].AssetClass);
-			}
-
-			for( const auto& PackageName : PackageNames )
-			{
-				TArray<FName> DependencyArray;
-				AssetRegistryModule.Get().GetDependencies(PackageName,DependencyArray);
-
-				FPackage* Package = FPackage::FindOrCreatePackage(PackageName);
-				for( const auto& DependencyName : DependencyArray )
-				{
-					// exclude '/Script/' as it clutters up things significantly.
-					if( !DependencyName.ToString().StartsWith(TEXT("/Script/")) )
-					{
-						FPackage* Dependency = FPackage::FindOrCreatePackage(DependencyName);
-						Package->DirectlyReferencing.Add(Dependency);
-						Dependency->DirectlyReferencedBy.Add(Package);
-					}
-				}
-			}
-
-			// 2 passes are required to deal with cycles.
-			for(const auto& PackageName : PackageNames)
-			{
-				FPackage* Package = FPackage::FindOrCreatePackage(PackageName);
-				Package->RecurseUpdateReferences();
-			}
-			for(const auto& PackageName : PackageNames)
-			{
-				FPackage* Package = FPackage::FindOrCreatePackage(PackageName);
-				Package->ResetUpdateHelper();
-			}
-			for(const auto& PackageName : PackageNames)
-			{
-				FPackage* Package = FPackage::FindOrCreatePackage(PackageName);
-				Package->RecurseUpdateReferences();
-			}
-		}
-
-		// Update inclusive size, asset class, and export to CSV, JSON, and GDF
-		{
-			TSharedPtr<FJsonObject> JsonRootObject = MakeShareable(new FJsonObject);
-			TArray< TSharedPtr<FJsonValue> > JsonPackages;
-
-			TArray<FPackage*> AllPackages;
-			FPackage::GetAllPackages(AllPackages);
-
-			for(auto Package : AllPackages)
-			{
-				Package->UpdateInclusiveSize();
-				JsonPackages.Add( MakeShareable(new FJsonValueObject(Package->ToJsonObject())) );
-			}
-			JsonRootObject->SetArrayField(TEXT("Packages"),JsonPackages);
-
-			FString JsonOutputString;
-			TSharedRef<TJsonWriter<TCHAR,TPrettyJsonPrintPolicy<TCHAR>>> JsonWriter = TJsonWriterFactory<TCHAR,TPrettyJsonPrintPolicy<TCHAR>>::Create(&JsonOutputString);
-			FJsonSerializer::Serialize(JsonRootObject.ToSharedRef(),JsonWriter);
-
-			FArchive* JsonFileWriter = IFileManager::Get().CreateFileWriter(*(FString(OutputFilenameBase)+TEXT(".json")));
-			if(JsonFileWriter)
-			{
-				JsonFileWriter->Logf(TEXT("%s"),*JsonOutputString);
-				JsonFileWriter->Close();
-				delete JsonFileWriter;
-			}
-
-			FArchive* CSVFileWriter = IFileManager::Get().CreateFileWriter(*(FString(OutputFilenameBase)+TEXT(".csv")));
-			if(CSVFileWriter)
-			{
-				CSVFileWriter->Logf(TEXT("class,name,inclusive,exclusive"));
-				for(auto Package : AllPackages)
-				{
-					FName ClassName = PackageToClassMap.FindRef(Package->Name);
-					CSVFileWriter->Logf(TEXT("%s,%s,%i,%i"),*ClassName.ToString(),*Package->Name.ToString(),Package->InclusiveSize,Package->ExclusiveSize);
-				}
-				CSVFileWriter->Close();
-				delete CSVFileWriter;
-				CSVFileWriter = NULL;
-			}
-
-			FArchive* GDFFileWriter = IFileManager::Get().CreateFileWriter(*(FString(OutputFilenameBase)+TEXT(".gdf")));
-			if(GDFFileWriter)
-			{
-				GDFFileWriter->Logf(TEXT("nodedef> name VARCHAR,label VARCHAR,inclusive DOUBLE,exclusive DOUBLE"));
-				GDFFileWriter->Logf(TEXT("0,root,0,0"));
-				for(auto Package : AllPackages)
-				{
-					GDFFileWriter->Logf(TEXT("%i,%s,%i,%i"),Package->Id,*Package->Name.ToString(),Package->InclusiveSize,Package->ExclusiveSize);
-				}
-				GDFFileWriter->Logf(TEXT("edgedef> node1 VARCHAR,node2 VARCHAR"));
-				// fake root to ensure spanning tree
-				for(auto Package : AllPackages)
-				{
-					GDFFileWriter->Logf(TEXT("0,%i"),Package->Id);
-				}
-				for(auto Package : AllPackages)
-				{
-					for( auto ReferencedPackage : Package->DirectlyReferencing )
-					{
-						GDFFileWriter->Logf(TEXT("%i,%i"),Package->Id,ReferencedPackage->Id);
-					}
-				}
-				GDFFileWriter->Close();
-				delete GDFFileWriter;
-				GDFFileWriter = NULL;
-			}
-		}
-		return true;
-	}
-	else
-	{
-		return false;
-	}
-}
-
 bool Repack(const FString& InputPakFile, const FString& OutputPakFile, const FPakCommandLineParameters& CmdLineParameters, const FKeyPair& SigningKey, const FAES::FAESKey& InEncryptionKey, bool bSigned)
 {
 	bool bResult = false;
@@ -2583,22 +2317,6 @@ bool ExecuteUnrealPak(const TCHAR* CmdLine)
 		return ListFilesInPak(*PakFilename, SizeFilter, bSigned);
 	}
 
-	FString ExportDependencyFilename;
-	if(FParse::Value(CmdLine,TEXT("ExportDependencies="),ExportDependencyFilename,false))
-	{
-		if(NonOptionArguments.Num() != 3)
-		{
-			UE_LOG(LogPakFile,Error,TEXT("Incorrect arguments. Expected: -ExportDependencies=<FileName> <PakFile> <GameName> <GameFolderName> [-Signed]"));
-			return false;
-		}
-
-		FString PakFilename = GetPakPath(*NonOptionArguments[0], false);
-		const TCHAR* GameName = *NonOptionArguments[1];
-		const TCHAR* GameFolderName = *NonOptionArguments[2];
-		bool bSigned = FParse::Param(CmdLine, TEXT("signed"));
-		return ExportDependencies(*PakFilename, GameName, GameFolderName, *ExportDependencyFilename, bSigned);
-	}
-
 	if (FParse::Param(CmdLine, TEXT("Diff")))
 	{
 		if(NonOptionArguments.Num() != 2)
diff --git a/Engine/Source/Editor/UnrealEd/Classes/Commandlets/ExportPakDependencies.h b/Engine/Source/Editor/UnrealEd/Classes/Commandlets/ExportPakDependencies.h
new file mode 100644
--- /dev/null
+++ b/Engine/Source/Editor/UnrealEd/Classes/Commandlets/ExportPakDependencies.h
@@ -0,0 +1,19 @@
+// Copyright 1998-2018 Epic Games, Inc. All Rights Reserved.
+
+#pragma once
+
+#include "CoreMinimal.h"
+#include "UObject/ObjectMacros.h"
+#include "Commandlets/Commandlet.h"
+#include "ExportPakDependencies.generated.h"
+
+UCLASS()
+class UExportPakDependenciesCommandlet : public UCommandlet
+{
+	GENERATED_UCLASS_BODY()
+	//~ Begin UCommandlet Interface
+	virtual int32 Main(const FString& Params) override;
+	//~ End UCommandlet Interface
+};
+
+
diff --git a/Engine/Source/Editor/UnrealEd/Private/Commandlets/ExportPakDependencies.cpp b/Engine/Source/Editor/UnrealEd/Private/Commandlets/ExportPakDependencies.cpp
new file mode 100644
--- /dev/null
+++ b/Engine/Source/Editor/UnrealEd/Private/Commandlets/ExportPakDependencies.cpp
@@ -0,0 +1,299 @@
+// Copyright 1998-2018 Epic Games, Inc. All Rights Reserved.
+
+#include "Commandlets/ExportPakDependencies.h"
+#include "AssetRegistryModule.h"
+#include "Containers/Set.h"
+#include "Containers/Array.h"
+#include "Dom/JsonValue.h"
+#include "Dom/JsonObject.h"
+#include "Serialization/JsonWriter.h"
+#include "Serialization/JsonSerializer.h"
+#include "HAL/FileManager.h"
+#include "IPlatformFilePak.h"
+#include "HAL/PlatformFilemanager.h"
+#include "Misc/CommandLine.h"
+#include "Misc/Paths.h"
+#include "Misc/App.h"
+
+struct FPackage
+{
+public:
+	TSet<FPackage*> DirectlyReferencing;
+	TSet<FPackage*> DirectlyReferencedBy;
+	TSet<FPackage*> AllReferences;
+
+	FName Name;
+	int64 InclusiveSize;
+	int64 ExclusiveSize;
+	int Id;
+
+private:
+	bool bUpdateHelper;
+	static TMap<FName,FPackage*> NameToPackageMap;
+
+public:
+	FPackage(const FName& InName,int InId)
+		: Name(InName)
+		, InclusiveSize(0)
+		, ExclusiveSize(0)
+		, Id(InId)
+		, bUpdateHelper(false)
+	{}
+
+	static FPackage* FindOrCreatePackage( FName PackageName )
+	{
+		static int Id = 1;
+		FPackage* Package = NameToPackageMap.FindRef(PackageName);
+		if(!Package)
+		{
+			Package = new FPackage(PackageName,Id++);
+			NameToPackageMap.Add(PackageName,Package);
+		}
+		return Package;
+	}
+
+	void ResetUpdateHelper()
+	{
+		bUpdateHelper = false;
+	}
+
+	void RecurseUpdateReferences()
+	{
+		if( !bUpdateHelper )
+		{
+			bUpdateHelper = true;
+			for( auto& DirectReference : DirectlyReferencing )
+			{
+				AllReferences.Add(DirectReference);
+				DirectReference->RecurseUpdateReferences();
+				AllReferences.Append(DirectReference->AllReferences);
+			}
+		}
+	}
+
+	void UpdateInclusiveSize()
+	{
+		InclusiveSize = ExclusiveSize;
+		for(auto& Reference : AllReferences)
+		{
+			InclusiveSize += Reference->ExclusiveSize;
+		}
+	}
+
+	static void GetAllPackages( TArray<FPackage*>& OutPackages )
+	{
+		OutPackages.Reset(NameToPackageMap.Num());
+		for( const auto& Entry : NameToPackageMap )
+		{
+			OutPackages.Add(Entry.Value);
+		}
+	}
+
+	TArray< TSharedPtr<FJsonValue> > ToJsonHelper( const TSet<FPackage*>& Packages )
+	{
+		TArray< TSharedPtr<FJsonValue> > JsonPackageNames;
+		for( const auto Package : Packages )
+		{
+			JsonPackageNames.Add(MakeShareable(new FJsonValueString(Package->Name.ToString())));
+		}
+		return JsonPackageNames;
+	}
+
+	TSharedPtr<FJsonObject> ToJsonObject()
+	{
+		TSharedPtr<FJsonObject> JsonPackageObject = MakeShareable(new FJsonObject);
+		
+		JsonPackageObject->SetStringField(TEXT("Name"),*Name.ToString());
+		JsonPackageObject->SetNumberField(TEXT("InclusiveSize"),InclusiveSize);
+		JsonPackageObject->SetNumberField(TEXT("ExclusiveSize"),ExclusiveSize);
+
+		JsonPackageObject->SetArrayField(TEXT("DirectlyReferencing"),ToJsonHelper(DirectlyReferencing));
+		JsonPackageObject->SetArrayField(TEXT("DirectlyReferencedBy"),ToJsonHelper(DirectlyReferencedBy));
+		JsonPackageObject->SetArrayField(TEXT("AllReferences"),ToJsonHelper(AllReferences));
+
+		return JsonPackageObject;
+	}
+};
+TMap<FName,FPackage*> FPackage::NameToPackageMap;
+
+bool ExportDependencies(const TCHAR * PakFilename, const TCHAR* GameName, const TCHAR* OutputFilenameBase, bool bSigned)
+{
+	// Example command line used for this tool
+	// C:\Development\BB\WEX\Saved\StagedBuilds\WindowsNoEditor\WorldExplorers\Content\Paks\WorldExplorers-WindowsNoEditor.pak WorldExplorers WEX -exportdependencies=c:\dvtemp\output -debug -NoAssetRegistryCache -ForceDependsGathering
+	
+	FPakFile PakFile(&FPlatformFileManager::Get().GetPlatformFile(), PakFilename, bSigned);
+
+	if(PakFile.IsValid())
+	{
+		// Get size information from PAK file.
+		{
+			TArray<FPakFile::FFileIterator> Records;
+			FString PakGameContentFolder = FString(GameName) + TEXT("/Content");
+			for(FPakFile::FFileIterator It(PakFile); It; ++It)
+			{
+				FString PackageName;
+				It.Filename().Split(TEXT("."),&PackageName,NULL);
+				int64 Size = It.Info().Size;			
+
+				if( PackageName.StartsWith(TEXT("Engine/Content")) )
+				{
+					PackageName = PackageName.Replace(TEXT("Engine/Content"),TEXT("/Engine"));
+				}
+				else if( PackageName.StartsWith(*PakGameContentFolder))
+				{
+					PackageName = PackageName.Replace(*PakGameContentFolder,TEXT("/Game"));
+				}
+
+				FPackage* Package = FPackage::FindOrCreatePackage(FName(*PackageName));
+				Package->ExclusiveSize += Size;
+			}
+		}
+
+		TMap<FName,FName> PackageToClassMap;
+
+		// Combine with dependency information from asset registry.
+		{
+			FAssetRegistryModule& AssetRegistryModule = FModuleManager::Get().LoadModuleChecked<FAssetRegistryModule>(TEXT("AssetRegistry"));
+		
+			AssetRegistryModule.Get().SearchAllAssets(true);
+			TArray<FAssetData> AssetData;		
+			AssetRegistryModule.Get().GetAllAssets(AssetData,true);	
+	
+			TSet<FName> PackageNames;
+			for( int i=0; i<AssetData.Num(); i++ )
+			{
+				PackageNames.Add(AssetData[i].PackageName);
+				PackageToClassMap.Add(AssetData[i].PackageName,AssetData[i].AssetClass);
+			}
+
+			for( const auto& PackageName : PackageNames )
+			{
+				TArray<FName> DependencyArray;
+				AssetRegistryModule.Get().GetDependencies(PackageName,DependencyArray);
+
+				FPackage* Package = FPackage::FindOrCreatePackage(PackageName);
+				for( const auto& DependencyName : DependencyArray )
+				{
+					// exclude '/Script/' as it clutters up things significantly.
+					if( !DependencyName.ToString().StartsWith(TEXT("/Script/")) )
+					{
+						FPackage* Dependency = FPackage::FindOrCreatePackage(DependencyName);
+						Package->DirectlyReferencing.Add(Dependency);
+						Dependency->DirectlyReferencedBy.Add(Package);
+					}
+				}
+			}
+
+			// 2 passes are required to deal with cycles.
+			for(const auto& PackageName : PackageNames)
+			{
+				FPackage* Package = FPackage::FindOrCreatePackage(PackageName);
+				Package->RecurseUpdateReferences();
+			}
+			for(const auto& PackageName : PackageNames)
+			{
+				FPackage* Package = FPackage::FindOrCreatePackage(PackageName);
+				Package->ResetUpdateHelper();
+			}
+			for(const auto& PackageName : PackageNames)
+			{
+				FPackage* Package = FPackage::FindOrCreatePackage(PackageName);
+				Package->RecurseUpdateReferences();
+			}
+		}
+
+		// Update inclusive size, asset class, and export to CSV, JSON, and GDF
+		{
+			TSharedPtr<FJsonObject> JsonRootObject = MakeShareable(new FJsonObject);
+			TArray< TSharedPtr<FJsonValue> > JsonPackages;
+
+			TArray<FPackage*> AllPackages;
+			FPackage::GetAllPackages(AllPackages);
+
+			for(auto Package : AllPackages)
+			{
+				Package->UpdateInclusiveSize();
+				JsonPackages.Add( MakeShareable(new FJsonValueObject(Package->ToJsonObject())) );
+			}
+			JsonRootObject->SetArrayField(TEXT("Packages"),JsonPackages);
+
+			FString JsonOutputString;
+			TSharedRef<TJsonWriter<TCHAR,TPrettyJsonPrintPolicy<TCHAR>>> JsonWriter = TJsonWriterFactory<TCHAR,TPrettyJsonPrintPolicy<TCHAR>>::Create(&JsonOutputString);
+			FJsonSerializer::Serialize(JsonRootObject.ToSharedRef(),JsonWriter);
+
+			FArchive* JsonFileWriter = IFileManager::Get().CreateFileWriter(*(FString(OutputFilenameBase)+TEXT(".json")));
+			if(JsonFileWriter)
+			{
+				JsonFileWriter->Logf(TEXT("%s"),*JsonOutputString);
+				JsonFileWriter->Close();
+				delete JsonFileWriter;
+			}
+
+			FArchive* CSVFileWriter = IFileManager::Get().CreateFileWriter(*(FString(OutputFilenameBase)+TEXT(".csv")));
+			if(CSVFileWriter)
+			{
+				CSVFileWriter->Logf(TEXT("class,name,inclusive,exclusive"));
+				for(auto Package : AllPackages)
+				{
+					FName ClassName = PackageToClassMap.FindRef(Package->Name);
+					CSVFileWriter->Logf(TEXT("%s,%s,%i,%i"),*ClassName.ToString(),*Package->Name.ToString(),Package->InclusiveSize,Package->ExclusiveSize);
+				}
+				CSVFileWriter->Close();
+				delete CSVFileWriter;
+				CSVFileWriter = NULL;
+			}
+
+			FArchive* GDFFileWriter = IFileManager::Get().CreateFileWriter(*(FString(OutputFilenameBase)+TEXT(".gdf")));
+			if(GDFFileWriter)
+			{
+				GDFFileWriter->Logf(TEXT("nodedef> name VARCHAR,label VARCHAR,inclusive DOUBLE,exclusive DOUBLE"));
+				GDFFileWriter->Logf(TEXT("0,root,0,0"));
+				for(auto Package : AllPackages)
+				{
+					GDFFileWriter->Logf(TEXT("%i,%s,%i,%i"),Package->Id,*Package->Name.ToString(),Package->InclusiveSize,Package->ExclusiveSize);
+				}
+				GDFFileWriter->Logf(TEXT("edgedef> node1 VARCHAR,node2 VARCHAR"));
+				// fake root to ensure spanning tree
+				for(auto Package : AllPackages)
+				{
+					GDFFileWriter->Logf(TEXT("0,%i"),Package->Id);
+				}
+				for(auto Package : AllPackages)
+				{
+					for( auto ReferencedPackage : Package->DirectlyReferencing )
+					{
+						GDFFileWriter->Logf(TEXT("%i,%i"),Package->Id,ReferencedPackage->Id);
+					}
+				}
+				GDFFileWriter->Close();
+				delete GDFFileWriter;
+				GDFFileWriter = NULL;
+			}
+		}
+		return true;
+	}
+	else
+	{
+		return false;
+	}
+}
+
+UExportPakDependenciesCommandlet::UExportPakDependenciesCommandlet(const FObjectInitializer& ObjectInitializer)
+	: Super(ObjectInitializer)
+{
+}
+
+int32 UExportPakDependenciesCommandlet::Main(const FString& Params)
+{
+	FString PakFilename;
+	FString ExportDependencyFilename;
+	if (!FParse::Value(FCommandLine::Get(), TEXT("PakFile="), PakFilename) || 
+		!FParse::Value(FCommandLine::Get(), TEXT("Output="), ExportDependencyFilename))
+	{
+		UE_LOG(LogPakFile,Error,TEXT("Incorrect arguments. Expected: -PakFile=<FileName> -Output=<FileName> [-Signed]"));
+		return false;
+	}
+
+	bool bSigned = FParse::Param(FCommandLine::Get(), TEXT("signed"));
+	return ExportDependencies(*PakFilename, FApp::GetProjectName(), *ExportDependencyFilename, bSigned)? 0 : 1;
+}
diff --git a/Engine/Source/Programs/UnrealPak/UnrealPak.Build.cs b/Engine/Source/Programs/UnrealPak/UnrealPak.Build.cs
--- a/Engine/Source/Programs/UnrealPak/UnrealPak.Build.cs
+++ b/Engine/Source/Programs/UnrealPak/UnrealPak.Build.cs
@@ -8,19 +8,13 @@ public class UnrealPak : ModuleRules
 	{
 		PublicIncludePaths.Add("Runtime/Launch/Public");
 
-		PrivateDependencyModuleNames.AddRange(new string[] { "Core", "CoreUObject", "PakFile", "Json", "Projects", "ApplicationCore", "PakFileUtilities" });
+		PrivateDependencyModuleNames.AddRange(new string[] { "Core", "PakFile", "Json", "Projects", "ApplicationCore", "PakFileUtilities" });
 
 		PrivateIncludePaths.Add("Runtime/Launch/Private");      // For LaunchEngineLoop.cpp include
 
         PrivateIncludePathModuleNames.AddRange(
             new string[] {
-                "AssetRegistry",
                 "Json"
         });
-
-        DynamicallyLoadedModuleNames.AddRange(
-            new string[] {
-                "AssetRegistry"
-        });
     }
 }
