Merging //UE4/Dev-Main to Dev-Sequencer (//UE4/Dev-Sequencer)

diff --git a/Engine/Source/Developer/PakFileUtilities/Private/KeyGenerator.cpp b/Engine/Source/Developer/PakFileUtilities/Private/KeyGenerator.cpp
deleted file mode 100644
diff --git a/Engine/Source/Developer/PakFileUtilities/Private/KeyGenerator.h b/Engine/Source/Developer/PakFileUtilities/Private/KeyGenerator.h
deleted file mode 100644
diff --git a/Engine/Source/Developer/PakFileUtilities/Private/PakFileUtilities.cpp b/Engine/Source/Developer/PakFileUtilities/Private/PakFileUtilities.cpp
--- a/Engine/Source/Developer/PakFileUtilities/Private/PakFileUtilities.cpp
+++ b/Engine/Source/Developer/PakFileUtilities/Private/PakFileUtilities.cpp
@@ -5,7 +5,6 @@
 #include "Misc/SecureHash.h"
 #include "Math/BigInt.h"
 #include "SignedArchiveWriter.h"
-#include "KeyGenerator.h"
 #include "Misc/AES.h"
 #include "Templates/UniquePtr.h"
 #include "Serialization/LargeMemoryWriter.h"
@@ -20,6 +19,71 @@
 #include "HAL/PlatformFilemanager.h"
 #include "Async/ParallelFor.h"
 
+/**
+ * Encryption keys: public and private
+ */
+struct FKeyPair
+{
+	/** Public decryption key */
+	FEncryptionKey PublicKey;
+	/** Private encryption key */
+	FEncryptionKey PrivateKey;
+
+	friend FArchive& operator<<(FArchive& Ar, FKeyPair& Pair)
+	{
+		Ar << Pair.PublicKey.Exponent;
+		Ar << Pair.PublicKey.Modulus;
+		Ar << Pair.PrivateKey.Exponent;
+		Ar << Pair.PrivateKey.Modulus;
+		return Ar;
+	}
+
+	bool IsValid() const
+	{
+		return !PrivateKey.Exponent.IsZero()
+			&& !PrivateKey.Modulus.IsZero()
+			&& !PublicKey.Exponent.IsZero()
+			&& !PublicKey.Modulus.IsZero();
+	}
+};
+
+bool TestKeys(FKeyPair& Pair)
+{
+	UE_LOG(LogPakFile, Display, TEXT("Testing signature keys."));
+
+	// Just some random values
+	static TEncryptionInt TestData[] =
+	{
+		11,
+		253,
+		128,
+		234,
+		56,
+		89,
+		34,
+		179,
+		29,
+		1024,
+		(int64)(MAX_int32),
+		(int64)(MAX_uint32)-1
+	};
+
+	for (int32 TestIndex = 0; TestIndex < ARRAY_COUNT(TestData); ++TestIndex)
+	{
+		TEncryptionInt EncryptedData = FEncryption::ModularPow(TestData[TestIndex], Pair.PrivateKey.Exponent, Pair.PrivateKey.Modulus);
+		TEncryptionInt DecryptedData = FEncryption::ModularPow(EncryptedData, Pair.PublicKey.Exponent, Pair.PublicKey.Modulus);
+		if (TestData[TestIndex] != DecryptedData)
+		{
+			UE_LOG(LogPakFile, Error, TEXT("Keys do not properly encrypt/decrypt data (failed test with %lld)"), TestData[TestIndex].ToInt());
+			return false;
+		}
+	}
+
+	UE_LOG(LogPakFile, Display, TEXT("Signature keys check completed successfuly."));
+
+	return true;
+}
+
 struct FNamedAESKey
 {
 	FString Name;
@@ -1199,34 +1263,6 @@ void PrepareEncryptionAndSigningKeys(const TCHAR* CmdLine, FKeyPair& OutSigningK
 			OutKeyChain.Add(NewKey.Guid, NewKey);
 			UE_LOG(LogPakFile, Display, TEXT("Parsed AES encryption key from command line."));
 		}
-
-		FString KeyFilename;
-		if (FParse::Value(CmdLine, TEXT("sign="), KeyFilename, false))
-		{
-			if (KeyFilename.StartsWith(TEXT("0x")))
-			{
-				TArray<FString> KeyValueText;
-				int32 NumParts = KeyFilename.ParseIntoArray(KeyValueText, TEXT("+"), true);
-				if (NumParts == 3)
-				{
-					OutSigningKey.PrivateKey.Exponent.Parse(KeyValueText[0]);
-					OutSigningKey.PrivateKey.Modulus.Parse(KeyValueText[1]);
-					OutSigningKey.PublicKey.Exponent.Parse(KeyValueText[2]);
-					OutSigningKey.PublicKey.Modulus = OutSigningKey.PrivateKey.Modulus;
-
-					UE_LOG(LogPakFile, Display, TEXT("Parsed signature keys from command line."));
-				}
-				else
-				{
-					UE_LOG(LogPakFile, Error, TEXT("Expected 3, got %d, when parsing %s"), KeyValueText.Num(), *KeyFilename);
-					OutSigningKey.PrivateKey.Exponent.Zero();
-				}
-			}
-			else if (!ReadKeysFromFile(*KeyFilename, OutSigningKey))
-			{
-				UE_LOG(LogPakFile, Error, TEXT("Unable to load signature keys %s."), *KeyFilename);
-			}
-		}
 	}
 
 	if (OutSigningKey.IsValid())
@@ -3341,10 +3377,6 @@ bool Repack(const FString& InputPakFile, const FString& OutputPakFile, const FPa
  *   -Sign=filename use the key pair in filename to sign a pak file, or: -sign=key_hex_values_separated_with_+, i.e: -sign=0x123456789abcdef+0x1234567+0x12345abc
  *    where the first number is the private key exponend, the second one is modulus and the third one is the public key exponent.
  *   -Signed use with -extract and -test to let the code know this is a signed pak
- *   -GenerateKeys=filename generates encryption key pair for signing a pak file
- *   -P=prime will use a predefined prime number for generating encryption key file
- *   -Q=prime same as above, P != Q, GCD(P, Q) = 1 (which is always true if they're both prime)
- *   -GeneratePrimeTable=filename generates a prime table for faster prime number generation (.inl file)
  *   -TableMax=number maximum prime number in the generated table (default is 10000)
  *
  * @param	ArgC	Command-line argument count
@@ -3389,27 +3421,6 @@ bool ExecuteUnrealPak(const TCHAR* CmdLine)
 		return Result;
 	}
 
-	FString KeyFilename;
-	if (FParse::Value(CmdLine, TEXT("GenerateKeys="), KeyFilename, false))
-	{
-		return GenerateKeys(*KeyFilename);
-	}
-
-	if (FParse::Value(CmdLine, TEXT("GeneratePrimeTable="), KeyFilename, false))
-	{
-		int64 MaxPrimeValue = 10000;
-		FParse::Value(CmdLine, TEXT("TableMax="), MaxPrimeValue);
-		GeneratePrimeNumberTable(MaxPrimeValue, *KeyFilename);
-		return true;
-	}
-
-	if (FParse::Param(CmdLine, TEXT("TestEncryption")))
-	{
-		void TestEncryption();
-		TestEncryption();
-		return true;
-	}
-
 	if (FParse::Param(CmdLine, TEXT("Test")))
 	{
 		if (NonOptionArguments.Num() != 1)
@@ -3773,8 +3784,6 @@ bool ExecuteUnrealPak(const TCHAR* CmdLine)
 	UE_LOG(LogPakFile, Error, TEXT("  UnrealPak <PakFilename> -Create=<ResponseFile> [Options]"));
 	UE_LOG(LogPakFile, Error, TEXT("  UnrealPak <PakFilename> -Dest=<MountPoint>"));
 	UE_LOG(LogPakFile, Error, TEXT("  UnrealPak <PakFilename> -Repack [-Output=Path] [-ExcludeDeleted] [Options]"));
-	UE_LOG(LogPakFile, Error, TEXT("  UnrealPak GenerateKeys=<KeyFilename>"));
-	UE_LOG(LogPakFile, Error, TEXT("  UnrealPak GeneratePrimeTable=<KeyFilename> [-TableMax=<N>]"));
 	UE_LOG(LogPakFile, Error, TEXT("  UnrealPak <PakFilename1> <PakFilename2> -diff"));
 	UE_LOG(LogPakFile, Error, TEXT("  UnrealPak <PakFolder> -AuditFiles [-OnlyDeleted] [-CSV=<filename>] [-order=<OrderingFile>] [-SortByOrdering]"));
 	UE_LOG(LogPakFile, Error, TEXT("  UnrealPak <PakFilename> -WhatsAtOffset [offset1] [offset2] [offset3] [...]"));
diff --git a/Engine/Source/Developer/PakFileUtilities/Private/Primes.inl b/Engine/Source/Developer/PakFileUtilities/Private/Primes.inl
deleted file mode 100644
diff --git a/Engine/Source/Developer/PakFileUtilities/Private/SignedArchiveWriter.cpp b/Engine/Source/Developer/PakFileUtilities/Private/SignedArchiveWriter.cpp
--- a/Engine/Source/Developer/PakFileUtilities/Private/SignedArchiveWriter.cpp
+++ b/Engine/Source/Developer/PakFileUtilities/Private/SignedArchiveWriter.cpp
@@ -99,64 +99,4 @@ int64 FSignedArchiveWriter::TotalSize()
 void FSignedArchiveWriter::Seek(int64 InPos)
 {
 	UE_LOG(LogPakFile, Fatal, TEXT("Seek is not supported in FSignedArchiveWriter."));
-}
-
-/**
- * Useful code for testing the encryption methods
- */
-void TestEncryption()
-{
-	FEncryptionKey PublicKey;
-	FEncryptionKey PrivateKey;
-	TEncryptionInt P(TEXT("0x21443BD2DD63E995403"));
-	TEncryptionInt Q(TEXT("0x28CBB6E5749AC65749"));
-	FEncryption::GenerateKeyPair(P, Q, PublicKey, PrivateKey);
-
-	// Generate random data
-	const int32 DataSize = 1024;
-	uint8* Data = (uint8*)FMemory::Malloc(DataSize);
-	for (int32 Index = 0; Index < DataSize; ++Index)
-	{
-		Data[Index] = (uint8)(Index % 255);
-	}
-
-	// Generate signature
-	FDecryptedSignature OriginalSignature;
-	FEncryptedSignature EncryptedSignature;
-	FDecryptedSignature DecryptedSignature;
-	OriginalSignature.Data = FCrc::MemCrc32(Data, DataSize);
-
-	// Encrypt signature with the private key
-	FEncryption::EncryptSignature(OriginalSignature, EncryptedSignature, PrivateKey);
-	FEncryption::DecryptSignature(EncryptedSignature, DecryptedSignature, PublicKey);
-
-	// Check
-	if (OriginalSignature == DecryptedSignature)
-	{
-		UE_LOG(LogPakFile, Display, TEXT("Keys match"));
-	}
-	else
-	{
-		UE_LOG(LogPakFile, Fatal, TEXT("Keys mismatched!"));
-	}
-
-	double OverallTime = 0.0;
-	double OverallNumTests = 0.0;
-	for (int32 TestIndex = 0; TestIndex < 10; ++TestIndex)
-	{
-		static const int64 NumTests = 500;
-		double Timer = FPlatformTime::Seconds();
-		{
-			for (int64 i = 0; i < NumTests; ++i)
-			{
-				FEncryption::DecryptSignature(EncryptedSignature, DecryptedSignature, PublicKey);
-			}
-		}
-		Timer = FPlatformTime::Seconds() - Timer;
-		OverallTime += Timer;
-		OverallNumTests += (double)NumTests;
-		UE_LOG(LogPakFile, Display, TEXT("%i signatures decrypted in %.4fs, Avg = %.4fs, OverallAvg = %.4fs"), NumTests, Timer, Timer / (double)NumTests, OverallTime / OverallNumTests);
-	}
-	
-	FMemory::Free(Data);
-}
+}
\ No newline at end of file
diff --git a/Engine/Source/Programs/AutomationTool/BuildGraph/Tasks/SubmitTask.cs b/Engine/Source/Programs/AutomationTool/BuildGraph/Tasks/SubmitTask.cs
--- a/Engine/Source/Programs/AutomationTool/BuildGraph/Tasks/SubmitTask.cs
+++ b/Engine/Source/Programs/AutomationTool/BuildGraph/Tasks/SubmitTask.cs
@@ -41,7 +41,7 @@ namespace AutomationTool.Tasks
 		public string Workspace;
 
 		/// <summary>
-		/// The stream for the workspace; defaults to the current stream. Ignored unless If the Workspace attribute is also specified.
+		/// The stream for the workspace; defaults to the current stream. Ignored unless the Workspace attribute is also specified.
 		/// </summary>
 		[TaskParameter(Optional = true)]
 		public string Stream;
diff --git a/Engine/Source/Programs/AutomationTool/Scripts/CopyBuildToStagingDirectory.Automation.cs b/Engine/Source/Programs/AutomationTool/Scripts/CopyBuildToStagingDirectory.Automation.cs
--- a/Engine/Source/Programs/AutomationTool/Scripts/CopyBuildToStagingDirectory.Automation.cs
+++ b/Engine/Source/Programs/AutomationTool/Scripts/CopyBuildToStagingDirectory.Automation.cs
@@ -343,8 +343,7 @@ public partial class Project : CommandUtils
 					if (bUFS)
 					{
 						List<FileReference> Files = SC.FindFilesToStage(InputDir, StageFilesSearch.AllDirectories);
-						Files.RemoveAll(x => x.HasExtension(".uasset") || x.HasExtension(".umap"));
-
+						Files.RemoveAll(x => x.HasExtension(".uasset") || x.HasExtension(".umap") || (SC.DedicatedServer && x.HasExtension(".mp4")));
 						SC.StageFiles(StagedFileType.UFS, InputDir, Files, OutputDir);
 					}
 					else
diff --git a/Engine/Source/Runtime/Advertising/IOS/IOSAdvertising/Private/IOSAdvertising.cpp b/Engine/Source/Runtime/Advertising/IOS/IOSAdvertising/Private/IOSAdvertising.cpp
--- a/Engine/Source/Runtime/Advertising/IOS/IOSAdvertising/Private/IOSAdvertising.cpp
+++ b/Engine/Source/Runtime/Advertising/IOS/IOSAdvertising/Private/IOSAdvertising.cpp
@@ -190,37 +190,37 @@ completion:^(BOOL finished)
 
 void FIOSAdvertisingProvider::ShowAdBanner( bool bShowOnBottomOfScreen, int32 /*AdID*/ ) 
 {
-#if __IPHONE_OS_VERSION_MIN_REQUIRED < __IPHONE_10_0
+#if __IPHONE_OS_VERSION_MIN_REQUIRED < __IPHONE_10_0 && !PLATFORM_TVOS
 	[[IOSAdvertising GetDelegate] performSelectorOnMainThread:@selector(ShowAdBanner:) withObject:[NSNumber numberWithBool : bShowOnBottomOfScreen] waitUntilDone : NO];
 #else
-    UE_LOG(LogIOS, Warning, TEXT("iAd has been deprecated as of IOS 10 and is this advertising provider is now deprecated in UE 4.21"));
+    UE_LOG(LogIOS, Warning, TEXT("iAd has been deprecated as of IOS 10 and this advertising provider is now deprecated in UE 4.21"));
 #endif
 }
 
 void FIOSAdvertisingProvider::HideAdBanner() 
 {
-#if __IPHONE_OS_VERSION_MIN_REQUIRED < __IPHONE_10_0
+#if __IPHONE_OS_VERSION_MIN_REQUIRED < __IPHONE_10_0 && !PLATFORM_TVOS
 	[[IOSAdvertising GetDelegate] performSelectorOnMainThread:@selector(UserHideAdBanner) withObject:nil waitUntilDone : NO];
 #else
-    UE_LOG(LogIOS, Warning, TEXT("iAd has been deprecated as of IOS 10 and is this advertising provider is now deprecated in UE 4.21"));
+    UE_LOG(LogIOS, Warning, TEXT("iAd has been deprecated as of IOS 10 and this advertising provider is now deprecated in UE 4.21"));
 #endif
 }
 
 void FIOSAdvertisingProvider::CloseAdBanner() 
 {
-#if __IPHONE_OS_VERSION_MIN_REQUIRED < __IPHONE_10_0
+#if __IPHONE_OS_VERSION_MIN_REQUIRED < __IPHONE_10_0 && !PLATFORM_TVOS
 	[[IOSAdvertising GetDelegate] performSelectorOnMainThread:@selector(CloseAd) withObject:nil waitUntilDone : NO];
 #else
-    UE_LOG(LogIOS, Warning, TEXT("iAd has been deprecated as of IOS 10 and is this advertising provider is now deprecated in UE 4.21"));
+    UE_LOG(LogIOS, Warning, TEXT("iAd has been deprecated as of IOS 10 and this advertising provider is now deprecated in UE 4.21"));
 #endif
 }
 
 int32 FIOSAdvertisingProvider::GetAdIDCount()
 {
-#if __IPHONE_OS_VERSION_MIN_REQUIRED < __IPHONE_10_0
+#if __IPHONE_OS_VERSION_MIN_REQUIRED < __IPHONE_10_0 && !PLATFORM_TVOS
 	return 1;
 #else
-    UE_LOG(LogIOS, Warning, TEXT("iAd has been deprecated as of IOS 10 and is this advertising provider is now deprecated in UE 4.21"));
+    UE_LOG(LogIOS, Warning, TEXT("iAd has been deprecated as of IOS 10 and this advertising provider is now deprecated in UE 4.21"));
     return 0;
 #endif
 }
diff --git a/Engine/Source/Runtime/Core/Private/HAL/MallocStomp.cpp b/Engine/Source/Runtime/Core/Private/HAL/MallocStomp.cpp
--- a/Engine/Source/Runtime/Core/Private/HAL/MallocStomp.cpp
+++ b/Engine/Source/Runtime/Core/Private/HAL/MallocStomp.cpp
@@ -101,23 +101,22 @@ void* FMallocStomp::Malloc(SIZE_T Size, uint32 Alignment)
 	static const SIZE_T AllocationDataSize = sizeof(FAllocationData);
 
 	const FAllocationData AllocData = { FullAllocationPointer, TotalAllocationSize, AlignedSize, SentinelExpectedValue };
-	FAllocationData* AllocDataPointer = nullptr;
 
 	if(bUseUnderrunMode)
 	{
 		const SIZE_T AlignedAllocationData = (Alignment > 0U) ? ((AllocationDataSize + Alignment - 1U) & -static_cast<int32>(Alignment)) : AllocationDataSize;
 		ReturnedPointer = reinterpret_cast<void*>(reinterpret_cast<uint8*>(FullAllocationPointer) + PageSize + AlignedAllocationData);
-		AllocDataPointer = reinterpret_cast<FAllocationData*>(reinterpret_cast<uint8*>(FullAllocationPointer) + PageSize);
+		void* AllocDataPointerStart = reinterpret_cast<FAllocationData*>(reinterpret_cast<uint8*>(FullAllocationPointer) + PageSize);
 
 #if PLATFORM_WINDOWS && MALLOC_STOMP_KEEP_VIRTUAL_MEMORY
 		// Commit physical pages to the used range, leaving the first page unmapped.
-		void* CommittedMemory = VirtualAlloc(AllocDataPointer, AllocFullPageSize, MEM_COMMIT, PAGE_READWRITE);
+		void* CommittedMemory = VirtualAlloc(AllocDataPointerStart, AllocFullPageSize, MEM_COMMIT, PAGE_READWRITE);
 		if (!CommittedMemory)
 		{
 			// Failed to allocate and commit physical memory pages. Report OOM.
 			FPlatformMemory::OnOutOfMemory(Size, Alignment);
 		}
-		check(CommittedMemory == AllocDataPointer);
+		check(CommittedMemory == AllocDataPointerStart);
 #else
 		// Page protect the first page, this will cause the exception in case the is an underrun.
 		FPlatformMemory::PageProtect(FullAllocationPointer, PageSize, false, false);
@@ -126,7 +125,6 @@ void* FMallocStomp::Malloc(SIZE_T Size, uint32 Alignment)
 	else
 	{
 		ReturnedPointer = reinterpret_cast<void*>(reinterpret_cast<uint8*>(FullAllocationPointer) + AllocFullPageSize - AlignedSize);
-		AllocDataPointer = reinterpret_cast<FAllocationData*>(reinterpret_cast<uint8*>(ReturnedPointer) - AllocationDataSize);
 
 #if PLATFORM_WINDOWS && MALLOC_STOMP_KEEP_VIRTUAL_MEMORY
 		// Commit physical pages to the used range, leaving the last page unmapped.
@@ -143,6 +141,7 @@ void* FMallocStomp::Malloc(SIZE_T Size, uint32 Alignment)
 #endif
 	} //-V773
 
+	FAllocationData* AllocDataPointer = reinterpret_cast<FAllocationData*>(reinterpret_cast<uint8*>(ReturnedPointer) - AllocationDataSize);
 	*AllocDataPointer = AllocData;
 
 	return ReturnedPointer;
diff --git a/Engine/Source/Runtime/LevelSequence/Private/LevelSequenceActor.cpp b/Engine/Source/Runtime/LevelSequence/Private/LevelSequenceActor.cpp
--- a/Engine/Source/Runtime/LevelSequence/Private/LevelSequenceActor.cpp
+++ b/Engine/Source/Runtime/LevelSequence/Private/LevelSequenceActor.cpp
@@ -52,7 +52,6 @@ ALevelSequenceActor::ALevelSequenceActor(const FObjectInitializer& Init)
 
 	// SequencePlayer must be a default sub object for it to be replicated correctly
 	SequencePlayer = Init.CreateDefaultSubobject<ULevelSequencePlayer>(this, "AnimationPlayer");
-	SequencePlayer->SetPlaybackClient(this);
 
 	bOverrideInstanceData = false;
 
@@ -63,6 +62,15 @@ ALevelSequenceActor::ALevelSequenceActor(const FObjectInitializer& Init)
 	bReplicatePlayback = false;
 }
 
+void ALevelSequenceActor::PostInitProperties()
+{
+	Super::PostInitProperties();
+
+	// Have to initialize this here as any properties set on default subobjects inside the constructor
+	// Get stomped by the CDO's properties when the constructor exits.
+	SequencePlayer->SetPlaybackClient(this);
+}
+
 bool ALevelSequenceActor::RetrieveBindingOverrides(const FGuid& InBindingId, FMovieSceneSequenceID InSequenceID, TArray<UObject*, TInlineAllocator<1>>& OutObjects) const
 {
 	return BindingOverrides->LocateBoundObjects(InBindingId, InSequenceID, OutObjects);
diff --git a/Engine/Source/Runtime/LevelSequence/Private/LevelSequenceBindingReference.cpp b/Engine/Source/Runtime/LevelSequence/Private/LevelSequenceBindingReference.cpp
--- a/Engine/Source/Runtime/LevelSequence/Private/LevelSequenceBindingReference.cpp
+++ b/Engine/Source/Runtime/LevelSequence/Private/LevelSequenceBindingReference.cpp
@@ -44,14 +44,14 @@ FLevelSequenceBindingReference::FLevelSequenceBindingReference(UObject* InObject
 
 UObject* FLevelSequenceBindingReference::Resolve(UObject* InContext, FName StreamedLevelAssetPath) const
 {
-	if (InContext->IsA<AActor>())
+	if (InContext && InContext->IsA<AActor>())
 	{
 		if (ExternalObjectPath.IsNull())
 		{
 			return FindObject<UObject>(InContext, *ObjectPath, false);
 		}
 	}
-	else if (InContext->IsA<ULevel>() && StreamedLevelAssetPath != NAME_None && ExternalObjectPath.GetAssetPathName() == StreamedLevelAssetPath)
+	else if (InContext && InContext->IsA<ULevel>() && StreamedLevelAssetPath != NAME_None && ExternalObjectPath.GetAssetPathName() == StreamedLevelAssetPath)
 	{
 		// ExternalObjectPath.GetSubPathString() specifies the path from the package (so includes PersistentLevel.) so we must do a FindObject from its outer
 		return FindObject<UObject>(InContext->GetOuter(), *ExternalObjectPath.GetSubPathString());
diff --git a/Engine/Source/Runtime/LevelSequence/Public/LevelSequenceActor.h b/Engine/Source/Runtime/LevelSequence/Public/LevelSequenceActor.h
--- a/Engine/Source/Runtime/LevelSequence/Public/LevelSequenceActor.h
+++ b/Engine/Source/Runtime/LevelSequence/Public/LevelSequenceActor.h
@@ -230,6 +230,7 @@ protected:
 
 	//~ Begin UObject interface
 	virtual bool ReplicateSubobjects(UActorChannel* Channel, FOutBunch* Bunch, FReplicationFlags *RepFlags) override;
+	virtual void PostInitProperties() override;
 	virtual void PostLoad() override;
 	//~ End UObject interface
 
diff --git a/Engine/Source/Runtime/OpenGLDrv/Private/OpenGLShaders.cpp b/Engine/Source/Runtime/OpenGLDrv/Private/OpenGLShaders.cpp
--- a/Engine/Source/Runtime/OpenGLDrv/Private/OpenGLShaders.cpp
+++ b/Engine/Source/Runtime/OpenGLDrv/Private/OpenGLShaders.cpp
@@ -2545,8 +2545,8 @@ void FDelayedEvictionContainer::Tick()
 			++It;
 			bMeasureEviction = true;
 			check(LinkedProgram->LRUInfo.EvictBucket == EvictBucketIndex);
+			LinkedProgram->LRUInfo.EvictBucket = -3; // Mark EvictBucket to indicated evicted from ProgramsToEvict, Prevent EvictProgram from attempting to remove again.
 			GetOpenGLProgramsCache().EvictProgram(LinkedProgram->Config.ProgramKey);
-			LinkedProgram->LRUInfo.EvictBucket = -3;
 			bMeasureEviction = false;
 		}
 	}
diff --git a/Engine/Source/Runtime/Renderer/Private/PostProcess/SceneRenderTargets.cpp b/Engine/Source/Runtime/Renderer/Private/PostProcess/SceneRenderTargets.cpp
--- a/Engine/Source/Runtime/Renderer/Private/PostProcess/SceneRenderTargets.cpp
+++ b/Engine/Source/Runtime/Renderer/Private/PostProcess/SceneRenderTargets.cpp
@@ -423,7 +423,23 @@ FIntPoint FSceneRenderTargets::ComputeDesiredSize(const FSceneViewFamily& ViewFa
 			!FMath::IsNearlyEqual(
 				(float)BufferSize.X / BufferSize.Y,
 				(float)DesiredBufferSize.X / DesiredBufferSize.Y);
-		bAllowDelayResize = bAllowDelayResize && !bAspectRatioChanged;
+
+		if (bAspectRatioChanged)
+		{
+			bAllowDelayResize = false;
+
+			// At this point we're assuming a simple output resize and forcing a hard swap so clear the history.
+			// If we don't the next frame will fail this check as the allocated aspect ratio will match the new
+			// frame's forced size so we end up looking through the history again, finding the previous old size
+			// and reallocating. Only after a few frames can the results actually settle when the history clears 
+			for (int32 i = 0; i < FrameSizeHistoryCount; ++i)
+			{
+				LargestDesiredSizes[i] = FIntPoint::ZeroValue;
+#if PREVENT_RENDERTARGET_SIZE_THRASHING
+				HistoryFlags[i] = 0;
+#endif
+			}
+		}
 	}
 
 	if(bAllowDelayResize)
diff --git a/Engine/Source/Runtime/Slate/Private/Widgets/Text/SlateTextBlockLayout.cpp b/Engine/Source/Runtime/Slate/Private/Widgets/Text/SlateTextBlockLayout.cpp
--- a/Engine/Source/Runtime/Slate/Private/Widgets/Text/SlateTextBlockLayout.cpp
+++ b/Engine/Source/Runtime/Slate/Private/Widgets/Text/SlateTextBlockLayout.cpp
@@ -209,24 +209,27 @@ void FSlateTextBlockLayout::UpdateTextHighlights(const FText& InHighlightText)
 	const TArray< FTextLayout::FLineModel >& LineModels = TextLayout->GetLineModels();
 
 	TArray<FTextRunRenderer> TextHighlights;
-	for(int32 LineIndex = 0; LineIndex < LineModels.Num(); ++LineIndex)
+	if (HighlightTextString.Len() > 0)
 	{
-		const FTextLayout::FLineModel& LineModel = LineModels[LineIndex];
-
-		int32 FindBegin = 0;
-		int32 CurrentHighlightBegin;
-		const int32 TextLength = LineModel.Text->Len();
-		while(FindBegin < TextLength && (CurrentHighlightBegin = LineModel.Text->Find(HighlightTextString, ESearchCase::IgnoreCase, ESearchDir::FromStart, FindBegin)) != INDEX_NONE)
+		for(int32 LineIndex = 0; LineIndex < LineModels.Num(); ++LineIndex)
 		{
-			FindBegin = CurrentHighlightBegin + HighlightTextLength;
+			const FTextLayout::FLineModel& LineModel = LineModels[LineIndex];
 
-			if(TextHighlights.Num() > 0 && TextHighlights.Last().LineIndex == LineIndex && TextHighlights.Last().Range.EndIndex == CurrentHighlightBegin)
-			{
-				TextHighlights[TextHighlights.Num() - 1] = FTextRunRenderer(LineIndex, FTextRange(TextHighlights.Last().Range.BeginIndex, FindBegin), TextHighlighter.ToSharedRef());
-			}
-			else
+			int32 FindBegin = 0;
+			int32 CurrentHighlightBegin;
+			const int32 TextLength = LineModel.Text->Len();
+			while(FindBegin < TextLength && (CurrentHighlightBegin = LineModel.Text->Find(HighlightTextString, ESearchCase::IgnoreCase, ESearchDir::FromStart, FindBegin)) != INDEX_NONE)
 			{
-				TextHighlights.Add(FTextRunRenderer(LineIndex, FTextRange(CurrentHighlightBegin, FindBegin), TextHighlighter.ToSharedRef()));
+				FindBegin = CurrentHighlightBegin + HighlightTextLength;
+
+				if(TextHighlights.Num() > 0 && TextHighlights.Last().LineIndex == LineIndex && TextHighlights.Last().Range.EndIndex == CurrentHighlightBegin)
+				{
+					TextHighlights[TextHighlights.Num() - 1] = FTextRunRenderer(LineIndex, FTextRange(TextHighlights.Last().Range.BeginIndex, FindBegin), TextHighlighter.ToSharedRef());
+				}
+				else
+				{
+					TextHighlights.Add(FTextRunRenderer(LineIndex, FTextRange(CurrentHighlightBegin, FindBegin), TextHighlighter.ToSharedRef()));
+				}
 			}
 		}
 	}
diff --git a/Engine/Source/Runtime/VulkanRHI/Private/Android/VulkanAndroidPlatform.cpp b/Engine/Source/Runtime/VulkanRHI/Private/Android/VulkanAndroidPlatform.cpp
--- a/Engine/Source/Runtime/VulkanRHI/Private/Android/VulkanAndroidPlatform.cpp
+++ b/Engine/Source/Runtime/VulkanRHI/Private/Android/VulkanAndroidPlatform.cpp
@@ -188,4 +188,20 @@ void FVulkanAndroidPlatform::OverridePlatformHandlers(bool bInit)
 	}
 }
 
+void FVulkanAndroidPlatform::BlockUntilWindowIsAwailable()
+{
+	void* WindowHandle = FAndroidWindow::GetHardwareWindow();
+	if (WindowHandle == nullptr)
+	{
+		// Sleep if the hardware window isn't currently available.
+		// The Window may not exist if the activity is pausing/resuming, in which case we make this thread wait
+		FPlatformMisc::LowLevelOutputDebugString(TEXT("Waiting for Native window in FVulkanAndroidPlatform::BlockUntilWindowIsAwailable"));
+		while (WindowHandle == nullptr)
+		{
+			FPlatformProcess::Sleep(0.001f);
+			WindowHandle = FAndroidWindow::GetHardwareWindow();
+		}
+	}
+}
+
 #endif
diff --git a/Engine/Source/Runtime/VulkanRHI/Private/Android/VulkanAndroidPlatform.h b/Engine/Source/Runtime/VulkanRHI/Private/Android/VulkanAndroidPlatform.h
--- a/Engine/Source/Runtime/VulkanRHI/Private/Android/VulkanAndroidPlatform.h
+++ b/Engine/Source/Runtime/VulkanRHI/Private/Android/VulkanAndroidPlatform.h
@@ -2,6 +2,8 @@
 
 #pragma once
 
+#include "RHI.h"
+
 #define VK_USE_PLATFORM_ANDROID_KHR					1
 
 #define VULKAN_HAS_PHYSICAL_DEVICE_PROPERTIES2		0
@@ -93,7 +95,9 @@ public:
 	static bool RequiresSwapchainGeneralInitialLayout() { return true; }
 
 	static bool RequiresWaitingForFrameCompletionEvent() { return false; }
-
+	
+	static void BlockUntilWindowIsAwailable();
+	
 protected:
 	static void* VulkanLib;
 	static bool bAttemptedLoad;
diff --git a/Engine/Source/Runtime/VulkanRHI/Private/VulkanGenericPlatform.h b/Engine/Source/Runtime/VulkanRHI/Private/VulkanGenericPlatform.h
--- a/Engine/Source/Runtime/VulkanRHI/Private/VulkanGenericPlatform.h
+++ b/Engine/Source/Runtime/VulkanRHI/Private/VulkanGenericPlatform.h
@@ -96,4 +96,7 @@ public:
 
 	// Ensure the last frame completed on the GPU
 	static bool RequiresWaitingForFrameCompletionEvent() { return true; }
+
+	// Blocks until hardware window is available
+	static void BlockUntilWindowIsAwailable() {};
 };
diff --git a/Engine/Source/Runtime/VulkanRHI/Private/VulkanViewport.cpp b/Engine/Source/Runtime/VulkanRHI/Private/VulkanViewport.cpp
--- a/Engine/Source/Runtime/VulkanRHI/Private/VulkanViewport.cpp
+++ b/Engine/Source/Runtime/VulkanRHI/Private/VulkanViewport.cpp
@@ -793,6 +793,9 @@ bool FVulkanViewport::Present(FVulkanCommandListContext* Context, FVulkanCmdBuff
 		Queue->Submit(CmdBuffer);
 	}
 
+	// Do not present until hardware window is available. On Android window could be destroyed while RHIT executes commands
+	FVulkanPlatform::BlockUntilWindowIsAwailable();
+
 	//Flush all commands
 	//check(0);
 
