- Reenabled encryption on bulk files with memory mapping - and then disallow mapping at runtime

diff --git a/Engine/Source/Developer/PakFileUtilities/Private/PakFileUtilities.cpp b/Engine/Source/Developer/PakFileUtilities/Private/PakFileUtilities.cpp
--- a/Engine/Source/Developer/PakFileUtilities/Private/PakFileUtilities.cpp
+++ b/Engine/Source/Developer/PakFileUtilities/Private/PakFileUtilities.cpp
@@ -846,12 +846,11 @@ void ProcessCommandLine(const TCHAR* CmdLine, const TArray<FString>& NonOptionAr
 			UE_LOG(LogPakFile, Log, TEXT("Added file Source: %s Dest: %s"), *Input.Source, *Input.Dest);
 
 			bool bIsUBulk = FPaths::GetExtension(Input.Source, true) == TEXT(".ubulk");
-			if (bIsUBulk && CmdLineParameters.AlignForMemoryMapping > 0 && (Input.bNeedsCompression || Input.bNeedEncryption))
+			if (bIsUBulk && CmdLineParameters.AlignForMemoryMapping > 0 && Input.bNeedsCompression)
 			{
 				// no compression for bulk aligned files because they are memory mapped
 				Input.bNeedsCompression = false;
-				Input.bNeedEncryption = false;
-				UE_LOG(LogPakFile, Warning, TEXT("Stripped compression or encryption from %s for memory mapping."), *Input.Dest);
+				UE_LOG(LogPakFile, Warning, TEXT("Stripped compression from %s for memory mapping."), *Input.Dest);
 			}
 			Entries.Add(Input);
 		}			
diff --git a/Engine/Source/Runtime/PakFile/Private/IPlatformFilePak.cpp b/Engine/Source/Runtime/PakFile/Private/IPlatformFilePak.cpp
--- a/Engine/Source/Runtime/PakFile/Private/IPlatformFilePak.cpp
+++ b/Engine/Source/Runtime/PakFile/Private/IPlatformFilePak.cpp
@@ -3825,9 +3825,9 @@ IMappedFileHandle* FPakPlatformFile::OpenMapped(const TCHAR* Filename)
 	FPakFile* PakEntry = nullptr;
 	if (FindFileInPakFiles(Filename, &PakEntry, &FileEntry) && PakEntry)
 	{
-		if (FileEntry.CompressionMethodIndex != 0)
+		if (FileEntry.CompressionMethodIndex != 0 || (FileEntry.Flags & FPakEntry::Flag_Encrypted) != 0)
 		{
-			// can't map compressed files
+			// can't map compressed or encrypted files
 			return nullptr;
 		}
 		FScopeLock Lock(&PakEntry->MappedFileHandleCriticalSection);
