Added validation to patch generation to detect when unsupported files are changed.

diff --git a/Engine/Source/Developer/PakFileUtilities/Private/PakFileUtilities.cpp b/Engine/Source/Developer/PakFileUtilities/Private/PakFileUtilities.cpp
--- a/Engine/Source/Developer/PakFileUtilities/Private/PakFileUtilities.cpp
+++ b/Engine/Source/Developer/PakFileUtilities/Private/PakFileUtilities.cpp
@@ -377,6 +377,7 @@ struct FPakCommandLineParameters
 	FString SourcePatchPakFilename;
 	FString SourcePatchDiffDirectory;
 	FString InputFinalPakFilename; // This is the resulting pak file we want to end up with after we generate the pak patch.  This is used instead of passing in the raw content.
+	FString ChangedFilesOutputFilename;
 	bool EncryptIndex;
 	bool UseCustomCompressor;
 	FGuid EncryptionKeyGuid;
@@ -866,10 +867,8 @@ void ProcessCommandLine(const TCHAR* CmdLine, const TArray<FString>& NonOptionAr
 
 	if (FParse::Value(CmdLine, TEXT("-create="), ResponseFile))
 	{
-		
-
 		CmdLineParameters.GeneratePatch = FParse::Value(CmdLine, TEXT("-generatepatch="), CmdLineParameters.SourcePatchPakFilename);
-
+		FParse::Value(CmdLine, TEXT("-outputchangedfiles="), CmdLineParameters.ChangedFilesOutputFilename);
 
 		bool bCompress = FParse::Param(CmdLine, TEXT("compress"));
 		bool bEncrypt = FParse::Param(CmdLine, TEXT("encrypt"));
@@ -3581,8 +3580,14 @@ void ApplyGapFilling(const TArray<FPakInputPair>& FilesToPak, const TArray<int64
 	UE_LOG(LogPakFile, SEEK_OPT_VERBOSITY, TEXT("Fragmentation final (primary files): %.2f%%"), FragmentationPercentageNewPrimary);
 }
 
-void RemoveIdenticalFiles( TArray<FPakInputPair>& FilesToPak, const FString& SourceDirectory, const TMap<FString, FFileInfo>& FileHashes, const FPatchSeekOptParams& SeekOptParams)
+void RemoveIdenticalFiles( TArray<FPakInputPair>& FilesToPak, const FString& SourceDirectory, const TMap<FString, FFileInfo>& FileHashes, const FPatchSeekOptParams& SeekOptParams, const FString ChangedFilesOutputFilename)
 {
+	FArchive* ChangedFilesArchive = nullptr;
+	if (ChangedFilesOutputFilename.IsEmpty() == false)
+	{
+		ChangedFilesArchive = IFileManager::Get().CreateFileWriter(*ChangedFilesOutputFilename);
+	}
+
 	FString HashFilename = SourceDirectory / TEXT("Hashes.txt");
 
 	if (IFileManager::Get().FileExists(*HashFilename) )
@@ -3625,6 +3630,8 @@ void RemoveIdenticalFiles( TArray<FPakInputPair>& FilesToPak, const FString& Sou
 	TArray<int64> FileSizes;
 	FileSizes.AddDefaulted(FilesToPak.Num());
 
+
+
     // Mark files to remove if they're unchanged
 	for (int i= 0; i<FilesToPak.Num(); i++)
 	{
@@ -3659,7 +3666,14 @@ void RemoveIdenticalFiles( TArray<FPakInputPair>& FilesToPak, const FString& Sou
 #endif
 
 		FString DestFilename = NewFile.Source;
-		if (!bForceInclude && FileIsIdentical(SourceFilename, DestFilename, FoundFileHash, &FileSizes[i]))
+		bool bFileIsIdentical = FileIsIdentical(SourceFilename, DestFilename, FoundFileHash, &FileSizes[i]);
+
+		if (ChangedFilesArchive && !bFileIsIdentical)
+		{
+			ChangedFilesArchive->Logf(TEXT("%s\n"),*DestFilename);
+		}
+
+		if (!bForceInclude && bFileIsIdentical)
 		{
 			UE_LOG(LogPakFile, Display, TEXT("Source file %s matches dest file %s and will not be included in patch"), *SourceFilename, *DestFilename);
 			// remove from the files to pak list
@@ -3700,6 +3714,13 @@ void RemoveIdenticalFiles( TArray<FPakInputPair>& FilesToPak, const FString& Sou
 }
 	int NumToRemove = FilesToPak.Num() - WriteIndex;
 	FilesToPak.RemoveAt(WriteIndex, NumToRemove, true);
+
+	if (ChangedFilesArchive)
+	{
+		ChangedFilesArchive->Close();
+		delete ChangedFilesArchive;
+		ChangedFilesArchive = nullptr;
+	}
 }
 
 void ProcessLegacyFileMoves( TArray<FPakInputPair>& InDeleteRecords, TMap<FString, FFileInfo>& InExistingPackagedFileHashes, const FString& InInputPath, const TArray<FPakInputPair>& InFilesToPak, int32 CurrentPatchChunkIndex )
@@ -4460,7 +4481,7 @@ bool ExecuteUnrealPak(const TCHAR* CmdLine)
 			FilesToAdd.Append(DeleteRecords);
 
 			// if we are generating a patch here we remove files which are already shipped...
-			RemoveIdenticalFiles(FilesToAdd, CmdLineParameters.SourcePatchDiffDirectory, SourceFileHashes, CmdLineParameters.SeekOptParams);
+			RemoveIdenticalFiles(FilesToAdd, CmdLineParameters.SourcePatchDiffDirectory, SourceFileHashes, CmdLineParameters.SeekOptParams, CmdLineParameters.ChangedFilesOutputFilename);
 		}
 
 
