Fixed crash in unrealpak in uncompressed non async code path.

diff --git a/Engine/Source/Developer/PakFileUtilities/Private/PakFileUtilities.cpp b/Engine/Source/Developer/PakFileUtilities/Private/PakFileUtilities.cpp
--- a/Engine/Source/Developer/PakFileUtilities/Private/PakFileUtilities.cpp
+++ b/Engine/Source/Developer/PakFileUtilities/Private/PakFileUtilities.cpp
@@ -34,7 +34,7 @@ IMPLEMENT_MODULE(FDefaultModuleImpl, PakFileUtilities);
 
 #define SEEK_OPT_VERBOSITY Display
 
-#define DETAILED_UNREALPAK_TIMING 1
+#define DETAILED_UNREALPAK_TIMING 0
 #if DETAILED_UNREALPAK_TIMING
 struct FUnrealPakScopeCycleCounter
 {
@@ -1904,6 +1904,10 @@ bool CreatePakFile(const TCHAR* Filename, TArray<FPakInputPair>& FilesToAdd, con
 
 		void Compress()
 		{
+			if (bIsComplete == true)
+			{
+				return;
+			}
 
 			//check if this file requested to be compression
 			OriginalFileSize = IFileManager::Get().FileSize(*FileToAdd->Source);
