Fixed alignment of non compressed files in pak files.
Also removed DDC config file.  Will need to figure out another way to do this when we enable DDC.
Graeme.Thornton

diff --git a/Engine/Programs/UnrealPak/Config/DefaultEngine.ini b/Engine/Programs/UnrealPak/Config/DefaultEngine.ini
deleted file mode 100644
diff --git a/Engine/Source/Developer/PakFileUtilities/Private/PakFileUtilities.cpp b/Engine/Source/Developer/PakFileUtilities/Private/PakFileUtilities.cpp
--- a/Engine/Source/Developer/PakFileUtilities/Private/PakFileUtilities.cpp
+++ b/Engine/Source/Developer/PakFileUtilities/Private/PakFileUtilities.cpp
@@ -403,6 +403,15 @@ public:
 #endif
 		return ReturnOrder;
 	}
+
+	void WriteOpenOrder(FArchive* Ar)
+	{
+		for (const auto& It : OrderMap)
+		{
+			Ar->Logf(TEXT("\"%s\" %d"), *It.Key, It.Value);
+		}
+	}
+
 private:
 	FString RemapLocalizationPathIfNeeded(const FString& PathLower, FString& OutRegion) const
 	{
@@ -1886,6 +1895,13 @@ bool CreatePakFile(const TCHAR* Filename, TArray<FPakInputPair>& FilesToAdd, con
 		
 		void Init(FPakInputPair* InFileToAdd, const TArray<FName>* InCompressionFormats, int32 InCompressionBlockSize, const TSet<FString>* InNoPluginCompressionExtensions) 
 		{
+			CompressionMethod = NAME_None;
+			if (InFileToAdd->bIsDeleteRecord)
+			{
+				// don't need to do anything for delete records
+				bIsComplete = true;
+				return;
+			}
 			FileToAdd = InFileToAdd;
 			NoPluginCompressionExtensions = InNoPluginCompressionExtensions;
 			CompressionFormats= InCompressionFormats;
@@ -1913,13 +1929,13 @@ bool CreatePakFile(const TCHAR* Filename, TArray<FPakInputPair>& FilesToAdd, con
 			OriginalFileSize = IFileManager::Get().FileSize(*FileToAdd->Source);
 			RealFileSize = OriginalFileSize + Entry.Info.GetSerializedSize(FPakInfo::PakFile_Version_Latest);
 
-			if (OriginalFileSize <= 0)
+			if (OriginalFileSize <= 0 || (FileToAdd->bNeedsCompression == false))
 			{
-				CompressionMethod = NAME_None;
+				// done, don't need to do anything else
+				Complete();
 				return;
 			}
 
-
 			bool bSomeCompressionSucceeded = false;
 			for (int32 MethodIndex = 0; MethodIndex < CompressionFormats->Num(); MethodIndex++)
 			{
@@ -2010,21 +2026,20 @@ bool CreatePakFile(const TCHAR* Filename, TArray<FPakInputPair>& FilesToAdd, con
 
 	for (int32 FileIndex = 0; FileIndex < FilesToAdd.Num(); FileIndex++)
 	{
-		bool bDeleted = FilesToAdd[FileIndex].bIsDeleteRecord;
-		//check if this file requested to be compression
-
-		if (FilesToAdd[FileIndex].bNeedsCompression)
+		AsyncCompressors[FileIndex].Init(&FilesToAdd[FileIndex], &CmdLineParameters.CompressionFormats, CmdLineParameters.CompressionBlockSize, &NoPluginCompressionExtensions);
+		if (bRunAsync)
 		{
-			AsyncCompressors[FileIndex].Init(&FilesToAdd[FileIndex], &CmdLineParameters.CompressionFormats, CmdLineParameters.CompressionBlockSize, &NoPluginCompressionExtensions);
-			if (bRunAsync)
+			if (FilesToAdd[FileIndex].bNeedsCompression)
 			{
 				(new FAutoDeleteAsyncTask<FRunCompressionTask>(&AsyncCompressors[FileIndex]))->StartBackgroundTask();
 			}
-		}
-		else
-		{
-			AsyncCompressors[FileIndex].CompressionMethod = NAME_None;
-			AsyncCompressors[FileIndex].Complete();
+			else
+			{
+				// call compress function inline 
+				// it won't do anything except for initialize some internal variables used in the non compressed path
+				// we don't want to pass these to a different thread as they may cause congestion with legitimate tasks
+				AsyncCompressors[FileIndex].Compress();
+			}
 		}
 	}
 
@@ -2073,16 +2088,15 @@ bool CreatePakFile(const TCHAR* Filename, TArray<FPakInputPair>& FilesToAdd, con
 			AsyncCompressors[FileIndex].Compress();
 		}
 
-
 		const FName& CompressionMethod = AsyncCompressors[FileIndex].CompressionMethod;
-		const int32 CompressionBlockSize = AsyncCompressors[FileIndex].CompressionBlockSize;
-		const int64 &RealFileSize = AsyncCompressors[FileIndex].RealFileSize;
-		const int64 &OriginalFileSize = AsyncCompressors[FileIndex].OriginalFileSize;
 		FCompressedFileBuffer& CompressedFileBuffer = AsyncCompressors[FileIndex].CompressedFileBuffer;
-		NewEntry.Info.CompressionMethodIndex = Info.GetCompressionMethodIndex(CompressionMethod);
 
 		if (!bDeleted)
 		{
+			const int64 &RealFileSize = AsyncCompressors[FileIndex].RealFileSize;
+			const int64 &OriginalFileSize = AsyncCompressors[FileIndex].OriginalFileSize;
+			const int32 CompressionBlockSize = AsyncCompressors[FileIndex].CompressionBlockSize;
+			NewEntry.Info.CompressionMethodIndex = Info.GetCompressionMethodIndex(CompressionMethod);
 
 			// Account for file system block size, which is a boundary we want to avoid crossing.
 			if (!bIsUAssetUExpPairUExp && // don't split uexp / uasset pairs
@@ -3231,12 +3245,12 @@ bool ExtractFilesFromPak(const TCHAR* InPakFilename, TMap<FString, FFileInfo>& I
 							{
 								UncompressCopyFile(*FileHandle, PakReader, Entry, PersistantCompressionBuffer, CompressionBufferSize, InKeyChain, PakFile);
 							}
-							UE_LOG(LogPakFile, Display, TEXT("Extracted \"%s\" to \"%s\"."), *It.Filename(), *DestFilename);
+							UE_LOG(LogPakFile, Display, TEXT("Extracted \"%s\" to \"%s\" Offset %d."), *It.Filename(), *DestFilename, Entry.Offset);
 							ExtractedCount++;
 
 							if (OutOrderMap != nullptr)
 							{
-								OutOrderMap->Add(DestFilename, OutOrderMap->Num());
+								OutOrderMap->Add(PakFile.GetMountPoint() / It.Filename(), It.GetIndexInPakFile());
 							}
 
 							if (OutEntries != nullptr)
@@ -4345,7 +4359,7 @@ bool Repack(const FString& InputPakFile, const FString& OutputPakFile, const FPa
 int32 NumberOfWorkerThreadsDesired()
 {
 	const int32 MaxThreads = 64;
-	const int32 NumberOfCores = FPlatformMisc::NumberOfCores();
+	const int32 NumberOfCores = FPlatformMisc::NumberOfCoresIncludingHyperthreads();
 	// need to spawn at least one worker thread (see FTaskGraphImplementation)
 	return FMath::Max(FMath::Min(NumberOfCores - 1, MaxThreads), 1);
 }
@@ -4535,15 +4549,16 @@ bool ExecuteUnrealPak(const TCHAR* CmdLine)
 	
 		if (NonOptionArguments.Num() != 2)
 		{
-			UE_LOG(LogPakFile, Error, TEXT("Incorrect arguments. Expected: -Extract <PakFile> <OutputPath> [-responsefile=<filename>]"));
+			UE_LOG(LogPakFile, Error, TEXT("Incorrect arguments. Expected: -Extract <PakFile> <OutputPath> [-responsefile=<outputresponsefilename> -order=<outputordermap>]"));
 			return false;
 		}
 
 		FString PakFilename = GetPakPath(*NonOptionArguments[0], false);
 
 		FString ResponseFileName;
+		FString OrderMapFileName;
 		bool bGenerateResponseFile = FParse::Value(CmdLine, TEXT("responseFile="), ResponseFileName);
-
+		bool bGenerateOrderMap = FParse::Value(CmdLine, TEXT("outorder="), OrderMapFileName);
 		bool bUseFilter = false;
 		FString Filter;
 		FString DestPath = NonOptionArguments[1];
@@ -4553,7 +4568,8 @@ bool ExecuteUnrealPak(const TCHAR* CmdLine)
 		TMap<FString, FFileInfo> EmptyMap;
 		TArray<FPakInputPair> ResponseContent;
 		TArray<FPakInputPair> DeletedContent;
-		if (ExtractFilesFromPak(*PakFilename, EmptyMap, *DestPath, bExtractToMountPoint, KeyChain, bUseFilter ? &Filter : nullptr, &ResponseContent, &DeletedContent) == false)
+		FPakOrderMap OrderMap;
+		if (ExtractFilesFromPak(*PakFilename, EmptyMap, *DestPath, bExtractToMountPoint, KeyChain, bUseFilter ? &Filter : nullptr, &ResponseContent, &DeletedContent, &OrderMap) == false)
 		{
 			return false;
 		}
@@ -4588,6 +4604,14 @@ bool ExecuteUnrealPak(const TCHAR* CmdLine)
 			ResponseArchive->Close();
 			delete ResponseArchive;
 		}
+		if (bGenerateOrderMap)
+		{
+			FArchive* OutputFile = IFileManager::Get().CreateFileWriter(*OrderMapFileName);
+			OutputFile->SetIsTextFormat(true);
+			OrderMap.WriteOpenOrder(OutputFile);
+			OutputFile->Close();
+			delete OutputFile;
+		}
 
 		return true;
 	}
diff --git a/Engine/Source/Runtime/PakFile/Public/IPlatformFilePak.h b/Engine/Source/Runtime/PakFile/Public/IPlatformFilePak.h
--- a/Engine/Source/Runtime/PakFile/Public/IPlatformFilePak.h
+++ b/Engine/Source/Runtime/PakFile/Public/IPlatformFilePak.h
@@ -899,6 +899,7 @@ public:
 
 		const FString& Filename() const		{ return CachedFilename; }
 		const FPakEntry& Info() const	{ return PakFile.Files[DirectoryIt.Value()]; }
+		const int32 GetIndexInPakFile() const { return DirectoryIt.Value(); }
 
 	private:
 		FORCEINLINE void AdvanceToValid()
