Fix static analysis warning.

diff --git a/Engine/Source/Developer/PakFileUtilities/Private/PakFileUtilities.cpp b/Engine/Source/Developer/PakFileUtilities/Private/PakFileUtilities.cpp
--- a/Engine/Source/Developer/PakFileUtilities/Private/PakFileUtilities.cpp
+++ b/Engine/Source/Developer/PakFileUtilities/Private/PakFileUtilities.cpp
@@ -627,7 +627,7 @@ struct FCompressedFileBuffer
 	int64 GetSerializedSizeEstimate() const
 	{
 		int64 Size = 0;
-		Size += sizeof(this);
+		Size += sizeof(*this);
 		Size += CompressedBlocks.Num() * sizeof(FPakCompressedBlock);
 		Size += CompressedBufferSize;
 		return Size;
