Add the option to force aligning files in pak files that are larger than block size and enable it
Aligning large files will significantly reduce patch size when doing delta patching on XboxOne.
However, it will also increase package size by about 5%.

diff --git a/Engine/Source/Developer/PakFileUtilities/Private/PakFileUtilities.cpp b/Engine/Source/Developer/PakFileUtilities/Private/PakFileUtilities.cpp
--- a/Engine/Source/Developer/PakFileUtilities/Private/PakFileUtilities.cpp
+++ b/Engine/Source/Developer/PakFileUtilities/Private/PakFileUtilities.cpp
@@ -481,6 +481,7 @@ struct FPakCommandLineParameters
 		, bPatchCompatibilityMode421(false)
 		, bFallbackOrderForNonUassetFiles(false)
 		, bAsyncCompression(false)
+		, bAlignFilesLargerThanBlock(false)
 	{
 	}
 
@@ -502,6 +503,7 @@ struct FPakCommandLineParameters
 	bool bPatchCompatibilityMode421;
 	bool bFallbackOrderForNonUassetFiles;
 	bool bAsyncCompression;
+	bool bAlignFilesLargerThanBlock;	// Align files that are larger than block size
 };
 
 struct FPakEntryPair
@@ -1075,6 +1077,11 @@ void ProcessCommandLine(const TCHAR* CmdLine, const TArray<FString>& NonOptionAr
 		CmdLineParameters.bSign = true;
 	}
 
+	if (FParse::Param(CmdLine, TEXT("AlignFilesLargerThanBlock")))
+	{
+		CmdLineParameters.bAlignFilesLargerThanBlock = true;
+	}
+
 	FString DesiredCompressionFormats;
 	// look for -compressionformats or -compressionformat on the commandline
 	if (FParse::Value(CmdLine, TEXT("-compressionformats="), DesiredCompressionFormats) || FParse::Value(CmdLine, TEXT("-compressionformat="), DesiredCompressionFormats))
@@ -2026,15 +2033,15 @@ bool CreatePakFile(const TCHAR* Filename, TArray<FPakInputPair>& FilesToAdd, con
 
 	for (int32 FileIndex = 0; FileIndex < FilesToAdd.Num(); FileIndex++)
 	{
-		AsyncCompressors[FileIndex].Init(&FilesToAdd[FileIndex], &CmdLineParameters.CompressionFormats, CmdLineParameters.CompressionBlockSize, &NoPluginCompressionExtensions);
-		if (bRunAsync)
-		{
+			AsyncCompressors[FileIndex].Init(&FilesToAdd[FileIndex], &CmdLineParameters.CompressionFormats, CmdLineParameters.CompressionBlockSize, &NoPluginCompressionExtensions);
+			if (bRunAsync)
+			{
 			if (FilesToAdd[FileIndex].bNeedsCompression)
 			{
 				(new FAutoDeleteAsyncTask<FRunCompressionTask>(&AsyncCompressors[FileIndex]))->StartBackgroundTask();
 			}
-			else
-			{
+		else
+		{
 				// call compress function inline 
 				// it won't do anything except for initialize some internal variables used in the non compressed path
 				// we don't want to pass these to a different thread as they may cause congestion with legitimate tasks
@@ -2100,11 +2107,11 @@ bool CreatePakFile(const TCHAR* Filename, TArray<FPakInputPair>& FilesToAdd, con
 
 			// Account for file system block size, which is a boundary we want to avoid crossing.
 			if (!bIsUAssetUExpPairUExp && // don't split uexp / uasset pairs
-				CmdLineParameters.FileSystemBlockSize > 0 && OriginalFileSize != INDEX_NONE && RealFileSize <= CmdLineParameters.FileSystemBlockSize)
+				CmdLineParameters.FileSystemBlockSize > 0 &&
+				OriginalFileSize != INDEX_NONE &&
+				(CmdLineParameters.bAlignFilesLargerThanBlock || RealFileSize <= CmdLineParameters.FileSystemBlockSize) &&
+				(NewEntryOffset / CmdLineParameters.FileSystemBlockSize) != ((NewEntryOffset + RealFileSize) / CmdLineParameters.FileSystemBlockSize)) // File crosses a block boundary
 			{
-				if ((NewEntryOffset / CmdLineParameters.FileSystemBlockSize) != ((NewEntryOffset + RealFileSize) / CmdLineParameters.FileSystemBlockSize))
-				{
-					//File crosses a block boundary, so align it to the beginning of the next boundary
 					int64 OldOffset = NewEntryOffset;
 					NewEntryOffset = AlignArbitrary(NewEntryOffset, CmdLineParameters.FileSystemBlockSize);
 					int64 PaddingRequired = NewEntryOffset - OldOffset;
@@ -2130,7 +2137,7 @@ bool CreatePakFile(const TCHAR* Filename, TArray<FPakInputPair>& FilesToAdd, con
 						check(PakFileHandle->Tell() == NewEntryOffset);
 					}
 				}
-			}
+
 			// Align bulk data
 			if (bIsMappedBulk && CmdLineParameters.AlignForMemoryMapping > 0 && OriginalFileSize != INDEX_NONE && !bDeleted)
 			{
