Handle UInt64 output properly when list file with unrealpak

diff --git a/Engine/Source/Developer/PakFileUtilities/Private/PakFileUtilities.cpp b/Engine/Source/Developer/PakFileUtilities/Private/PakFileUtilities.cpp
--- a/Engine/Source/Developer/PakFileUtilities/Private/PakFileUtilities.cpp
+++ b/Engine/Source/Developer/PakFileUtilities/Private/PakFileUtilities.cpp
@@ -2493,7 +2493,7 @@ bool ListFilesInPak(const TCHAR * InPakFilename, int64 SizeFilter, bool bInclude
 				bool bWasCompressed = Entry.CompressionMethodIndex != 0;
 
 				Lines.Add( FString::Printf(
-					TEXT("%s%s, %d, %d, %s, %s, %s, %d"),
+					TEXT("%s%s, %lld, %lld, %s, %s, %s, %d"),
 					*MountPoint, *It.Filename(),
 					Entry.Offset, Entry.Size,
 					*BytesToHex(Entry.Hash, sizeof(Entry.Hash)),
