Added an option to unrealpak which forces all compressed files inside a pak to be compressed, even if the compression results in a larger file size

diff --git a/Engine/Source/Developer/PakFileUtilities/Private/PakFileUtilities.cpp b/Engine/Source/Developer/PakFileUtilities/Private/PakFileUtilities.cpp
--- a/Engine/Source/Developer/PakFileUtilities/Private/PakFileUtilities.cpp
+++ b/Engine/Source/Developer/PakFileUtilities/Private/PakFileUtilities.cpp
@@ -482,6 +482,7 @@ struct FPakCommandLineParameters
 		, bFallbackOrderForNonUassetFiles(false)
 		, bAsyncCompression(false)
 		, bAlignFilesLargerThanBlock(false)
+		, bForceCompress(false)
 	{
 	}
 
@@ -504,6 +505,7 @@ struct FPakCommandLineParameters
 	bool bFallbackOrderForNonUassetFiles;
 	bool bAsyncCompression;
 	bool bAlignFilesLargerThanBlock;	// Align files that are larger than block size
+	bool bForceCompress; // Force all files that request compression to be compressed, even if that results in a larger file size
 };
 
 struct FPakEntryPair
@@ -1082,6 +1084,11 @@ void ProcessCommandLine(const TCHAR* CmdLine, const TArray<FString>& NonOptionAr
 		CmdLineParameters.bAlignFilesLargerThanBlock = true;
 	}
 
+	if (FParse::Param(CmdLine, TEXT("ForceCompress")))
+	{
+		CmdLineParameters.bForceCompress = true;
+	}
+
 	FString DesiredCompressionFormats;
 	// look for -compressionformats or -compressionformat on the commandline
 	if (FParse::Value(CmdLine, TEXT("-compressionformats="), DesiredCompressionFormats) || FParse::Value(CmdLine, TEXT("-compressionformat="), DesiredCompressionFormats))
@@ -1809,6 +1816,8 @@ FArchive* CreatePakWriter(const TCHAR* Filename, const FKeyChain& InKeyChain, bo
 	return Writer;
 }
 
+TAtomic<int64> GTotalFilesWithPoorForcedCompression = 0;
+TAtomic<int64> GTotalExtraMemoryForPoorForcedCompression = 0;
 
 bool CreatePakFile(const TCHAR* Filename, TArray<FPakInputPair>& FilesToAdd, const FPakCommandLineParameters& CmdLineParameters, const FKeyChain& InKeyChain)
 {
@@ -1898,9 +1907,10 @@ bool CreatePakFile(const TCHAR* Filename, TArray<FPakInputPair>& FilesToAdd, con
 		int32 CompressionBlockSize;
 		int64 RealFileSize;
 		int64 OriginalFileSize;
+		bool bForceCompress = false;
 		volatile bool bIsComplete;
 		
-		void Init(FPakInputPair* InFileToAdd, const TArray<FName>* InCompressionFormats, int32 InCompressionBlockSize, const TSet<FString>* InNoPluginCompressionExtensions) 
+		void Init(FPakInputPair* InFileToAdd, const FPakCommandLineParameters& InParams, const TSet<FString>* InNoPluginCompressionExtensions)
 		{
 			CompressionMethod = NAME_None;
 			if (InFileToAdd->bIsDeleteRecord)
@@ -1909,10 +1919,11 @@ bool CreatePakFile(const TCHAR* Filename, TArray<FPakInputPair>& FilesToAdd, con
 				bIsComplete = true;
 				return;
 			}
+			bForceCompress = InParams.bForceCompress;
 			FileToAdd = InFileToAdd;
 			NoPluginCompressionExtensions = InNoPluginCompressionExtensions;
-			CompressionFormats= InCompressionFormats;
-			CompressionBlockSize = InCompressionBlockSize;	
+			CompressionFormats = &InParams.CompressionFormats;
+			CompressionBlockSize = InParams.CompressionBlockSize;
 			bIsComplete = false;
 		}
 		
@@ -1964,7 +1975,9 @@ bool CreatePakFile(const TCHAR* Filename, TArray<FPakInputPair>& FilesToAdd, con
 					// if we still save 64KB it's probably worthwhile compressing, as that saves a file read operation in the runtime.
 								// TODO: drive this threshold from the command line
 					float PercentLess = ((float)CompressedFileBuffer.TotalCompressedSize / (OriginalFileSize / 100.f));
-					if (PercentLess > 90.f && (OriginalFileSize - CompressedFileBuffer.TotalCompressedSize) < 65536)
+					const bool bNotEnoughCompression = PercentLess > 90.f && (OriginalFileSize - CompressedFileBuffer.TotalCompressedSize) < 6553;
+					const bool bIsLastCompressionFormat = MethodIndex == CompressionFormats->Num() - 1;
+					if (bNotEnoughCompression && (!bForceCompress || !bIsLastCompressionFormat))
 					{
 						// compression did not succeed, we can try the next format, so do nothing here
 					}
@@ -1976,6 +1989,15 @@ bool CreatePakFile(const TCHAR* Filename, TArray<FPakInputPair>& FilesToAdd, con
 
 						// at this point, we have successfully compressed the file, no need to continue
 						bSomeCompressionSucceeded = true;
+
+						if (bNotEnoughCompression)
+						{
+							// None of the compression formats were good enough, but we were under instructions to use some form of compression
+							// This was likely to aid with runtime error checking on mobile devices. Record how many cases of this we encountered and
+							// How much the size difference was
+							GTotalFilesWithPoorForcedCompression++;
+							GTotalExtraMemoryForPoorForcedCompression += FMath::Max(CompressedFileBuffer.TotalCompressedSize - OriginalFileSize, (int64)0);
+						}
 					}
 				}
 
@@ -2033,7 +2055,7 @@ bool CreatePakFile(const TCHAR* Filename, TArray<FPakInputPair>& FilesToAdd, con
 
 	for (int32 FileIndex = 0; FileIndex < FilesToAdd.Num(); FileIndex++)
 	{
-			AsyncCompressors[FileIndex].Init(&FilesToAdd[FileIndex], &CmdLineParameters.CompressionFormats, CmdLineParameters.CompressionBlockSize, &NoPluginCompressionExtensions);
+			AsyncCompressors[FileIndex].Init(&FilesToAdd[FileIndex], CmdLineParameters, &NoPluginCompressionExtensions);
 			if (bRunAsync)
 			{
 			if (FilesToAdd[FileIndex].bNeedsCompression)
@@ -2393,6 +2415,11 @@ bool CreatePakFile(const TCHAR* Filename, TArray<FPakInputPair>& FilesToAdd, con
 		}
 
 		UE_LOG(LogPakFile, Display, TEXT("Used compression formats (in priority order) '%s'"), *UsedCompressionFormatsString);
+
+		if (GTotalFilesWithPoorForcedCompression > 0)
+		{
+			UE_LOG(LogPakFile, Display, TEXT("Num files forcibly compressed due to -forcecompress option: %i, using %i bytes extra"), (int64)GTotalFilesWithPoorForcedCompression, (int64)GTotalExtraMemoryForPoorForcedCompression);
+		}
 	}
 
 	if (TotalEncryptedDataSize)
