Fix error: copying variable of type 'TAtomic<int64>' (aka 'TAtomic<long long>') invokes deleted constructor
#rnx

diff --git a/Engine/Source/Developer/PakFileUtilities/Private/PakFileUtilities.cpp b/Engine/Source/Developer/PakFileUtilities/Private/PakFileUtilities.cpp
--- a/Engine/Source/Developer/PakFileUtilities/Private/PakFileUtilities.cpp
+++ b/Engine/Source/Developer/PakFileUtilities/Private/PakFileUtilities.cpp
@@ -1816,8 +1816,8 @@ FArchive* CreatePakWriter(const TCHAR* Filename, const FKeyChain& InKeyChain, bo
 	return Writer;
 }
 
-TAtomic<int64> GTotalFilesWithPoorForcedCompression = 0;
-TAtomic<int64> GTotalExtraMemoryForPoorForcedCompression = 0;
+TAtomic<int64> GTotalFilesWithPoorForcedCompression(0);
+TAtomic<int64> GTotalExtraMemoryForPoorForcedCompression(0);
 
 bool CreatePakFile(const TCHAR* Filename, TArray<FPakInputPair>& FilesToAdd, const FPakCommandLineParameters& CmdLineParameters, const FKeyChain& InKeyChain)
 {
