unrealpak: Fix a typo where I accidentally removed the last digit of a constant

#rnx

diff --git a/Engine/Source/Developer/PakFileUtilities/Private/PakFileUtilities.cpp b/Engine/Source/Developer/PakFileUtilities/Private/PakFileUtilities.cpp
--- a/Engine/Source/Developer/PakFileUtilities/Private/PakFileUtilities.cpp
+++ b/Engine/Source/Developer/PakFileUtilities/Private/PakFileUtilities.cpp
@@ -1975,7 +1975,7 @@ bool CreatePakFile(const TCHAR* Filename, TArray<FPakInputPair>& FilesToAdd, con
 					// if we still save 64KB it's probably worthwhile compressing, as that saves a file read operation in the runtime.
 								// TODO: drive this threshold from the command line
 					float PercentLess = ((float)CompressedFileBuffer.TotalCompressedSize / (OriginalFileSize / 100.f));
-					const bool bNotEnoughCompression = PercentLess > 90.f && (OriginalFileSize - CompressedFileBuffer.TotalCompressedSize) < 6553;
+					const bool bNotEnoughCompression = (PercentLess > 90.f) && ((OriginalFileSize - CompressedFileBuffer.TotalCompressedSize) < 65536);
 					const bool bIsLastCompressionFormat = MethodIndex == CompressionFormats->Num() - 1;
 					if (bNotEnoughCompression && (!bForceCompress || !bIsLastCompressionFormat))
 					{
