Optimize Pak file indices by computing the PathHashIndex in unrealpak rather than using strings at runtime.
This also removes the need to shrink the string indices for performance, which was triggering a bug due to unsynchronized multithreaded access.

diff --git a/Engine/Config/BaseEngine.ini b/Engine/Config/BaseEngine.ini
--- a/Engine/Config/BaseEngine.ini
+++ b/Engine/Config/BaseEngine.ini
@@ -1959,6 +1959,14 @@ StartingBeaconPort=15000
 +ExtensionsToNotUsePluginCompression=ini
 +ExtensionsToNotUsePluginCompression=icu
 +ExtensionsToNotUsePluginCompression=res
++DirectoryIndexKeepFiles="*/Config/Tags/*"
++DirectoryIndexKeepFiles="*/Content/Localization/*"
++DirectoryIndexKeepFiles="*/Content/Internationalization/*"
++DirectoryIndexKeepFiles="*/Content/Movies/*"
++DirectoryIndexKeepFiles="*/Plugins/*.uplugin*"
++DirectoryIndexKeepFiles="*/*.ini*"
++DirectoryIndexKeepEmptyDirectories=*/Plugins/*
++IndexValidationIgnore="*/Saved/Logs/*"
 
 [/Script/GameplayDebugger.GameplayDebuggingReplicator]
 MaxEQSQueries=5
diff --git a/Engine/Plugins/Online/OnlineFramework/Source/Hotfix/Private/OnlineHotfixManager.cpp b/Engine/Plugins/Online/OnlineFramework/Source/Hotfix/Private/OnlineHotfixManager.cpp
--- a/Engine/Plugins/Online/OnlineFramework/Source/Hotfix/Private/OnlineHotfixManager.cpp
+++ b/Engine/Plugins/Online/OnlineFramework/Source/Hotfix/Private/OnlineHotfixManager.cpp
@@ -1088,32 +1088,45 @@ bool UOnlineHotfixManager::HotfixIniFile(const FString& FileName, const FString&
 
 bool UOnlineHotfixManager::HotfixPakFile(const FCloudFileHeader& FileHeader)
 {
-	if (!FCoreDelegates::OnMountPak.IsBound())
+	if (!FCoreDelegates::MountPak.IsBound())
 	{
-		UE_LOG(LogHotfixManager, Error, TEXT("PAK file (%s) could not be mounted because OnMountPak is not bound"), *FileHeader.FileName);
+		UE_LOG(LogHotfixManager, Error, TEXT("PAK file (%s) could not be mounted because MountPak is not bound"), *FileHeader.FileName);
 		return false;
 	}
 	FString PakLocation = FString::Printf(TEXT("%s/%s"), *GetCachedDirectory(), *FileHeader.DLName);
-	FPakFileVisitor Visitor;
-	if (FCoreDelegates::OnMountPak.Execute(PakLocation, 0, &Visitor))
+	if (IPakFile* PakFile = FCoreDelegates::MountPak.Execute(PakLocation, 0))
 	{
 		MountedPakFiles.Add(FileHeader.DLName);
 		UE_LOG(LogHotfixManager, Log, TEXT("Hotfix mounted PAK file (%s)"), *FileHeader.FileName);
 		int32 NumInisReloaded = 0;
 		const double StartTime = FPlatformTime::Seconds();
+
+		// Iterate through the the pak file's contents for INI and asset reloading.
 		TArray<FString> IniList;
-		// Iterate through the pak file's contents for INI and asset reloading
+		FPakFileVisitor Visitor;
+		PakFile->PakVisitPrunedFilenames(Visitor);
 		for (const FString& InternalPakFileName : Visitor.Files)
 		{
 			if (InternalPakFileName.EndsWith(TEXT(".ini")))
 			{
 				IniList.Add(InternalPakFileName);
 			}
-			else if (!bHotfixNeedsMapReload && InternalPakFileName.EndsWith(FPackageName::GetMapPackageExtension()))
+		}
+
+		// Iterate through all loaded maps and see if they have patches in the pak file and therefore this hotfix needs to reload a map
+		for (TObjectIterator<UPackage> it; !bHotfixNeedsMapReload && it; ++it)
+		{
+			UPackage* Package = *it;
+			if (Package && Package->ContainsMap())
 			{
-				bHotfixNeedsMapReload = IsMapLoaded(InternalPakFileName);
+				const FString FileName = FPackageName::LongPackageNameToFilename(Package->FileName.ToString(), FPackageName::GetMapPackageExtension());
+				if (PakFile->PakContains(FileName))
+				{
+					bHotfixNeedsMapReload = true;
+				}
 			}
 		}
+
 		// Sort the INIs so they are processed consistently
 		IniList.Sort<FHotfixFileSortPredicate>(FHotfixFileSortPredicate(PlatformPrefix, ServerPrefix, DefaultPrefix));
 		// Now process the INIs in sorted order
diff --git a/Engine/Source/Developer/PakFileUtilities/Private/PakFileUtilities.cpp b/Engine/Source/Developer/PakFileUtilities/Private/PakFileUtilities.cpp
--- a/Engine/Source/Developer/PakFileUtilities/Private/PakFileUtilities.cpp
+++ b/Engine/Source/Developer/PakFileUtilities/Private/PakFileUtilities.cpp
@@ -12,6 +12,7 @@
 #include "Serialization/JsonSerializer.h"
 #include "Misc/Base64.h"
 #include "Misc/Compression.h"
+#include "Misc/Fnv.h"
 #include "Features/IModularFeatures.h"
 #include "Misc/CoreDelegates.h"
 #include "Misc/FileHelper.h"
@@ -23,7 +24,6 @@
 #include "DerivedDataCacheInterface.h"
 #include "Serialization/MemoryReader.h"
 #include "Serialization/MemoryWriter.h"
-#include "Misc/DataDrivenPlatformInfoRegistry.h"
 
 IMPLEMENT_MODULE(FDefaultModuleImpl, PakFileUtilities);
 
@@ -277,7 +277,10 @@ private:
 };
 
 /**
-* Defines the order mapping for files within a pak
+* Defines the order mapping for files within a pak.
+* When read from the files present in the pak, Indexes will be [0,NumFiles).  This is important for detecting gaps in the order between adjacent files in a patch .pak.
+* For new files being added into the pak, the values can be arbitrary, and will be usable only for relative order in an output list.
+* Due to the arbitrary values for new files, the FPakOrderMap can contain files with duplicate Order values.
  */
 class FPakOrderMap
 {
@@ -286,16 +289,52 @@ public:
 		: MaxPrimaryOrderIndex(MAX_uint64)
 	{}
 
+	void Empty()
+	{
+		OrderMap.Empty();
+		MaxPrimaryOrderIndex = MAX_uint64;
+	}
+
 	int32 Num() const
 	{
 		return OrderMap.Num();
 	}
 
+	/** Add the given filename with the given Sorting Index */
 	void Add(const FString& Filename, uint64 Index)
 	{
 		OrderMap.Add(Filename, Index);
 	}
 
+	/**
+	 * Add the given filename with the given Offset interpreted as Offset in bytes in the Pak File.  This version of Add is only useful when all Adds are done by offset, and are converted
+	 * into Sorting Indexes at the end by a call to ConvertOffsetsToOrder
+	 */
+	void AddOffset(const FString& Filename, uint64 Offset)
+	{
+		OrderMap.Add(Filename, Offset);
+	}
+
+	/** Remaps all the current values in the OrderMap onto [0, NumEntries).  Useful to convert from Offset in Pak file bytes into an Index sorted by Offset */
+	void ConvertOffsetsToOrder()
+	{
+		TArray<TPair<FString, uint64>> FilenameAndOffsets;
+		for (auto& FilenameAndOffset : OrderMap)
+		{
+			FilenameAndOffsets.Add(FilenameAndOffset);
+		}
+		FilenameAndOffsets.Sort([](const TPair<FString, uint64>& A, const TPair<FString, uint64>& B)
+		{
+			return A.Value < B.Value;
+		});
+		int64 Index = 0;
+		for (auto& FilenameAndOffset : FilenameAndOffsets)
+		{
+			OrderMap[FilenameAndOffset.Key] = Index;
+			++Index;
+		}
+	}
+
 	bool ProcessOrderFile(const TCHAR* ResponseFile, bool bSecondaryOrderFile = false)
 	{
 		int32 OrderOffset = 0;
@@ -407,6 +446,7 @@ public:
 
 	void WriteOpenOrder(FArchive* Ar)
 	{
+		OrderMap.ValueSort([](const uint64& A, const uint64& B) { return A < B; });
 		for (const auto& It : OrderMap)
 		{
 			Ar->Logf(TEXT("\"%s\" %d"), *It.Key, It.Value);
@@ -484,7 +524,6 @@ struct FPakCommandLineParameters
 		, bAsyncCompression(false)
 		, bAlignFilesLargerThanBlock(false)
 		, bForceCompress(false)
-		, bAllowForIndexUnload(false)
 	{
 	}
 
@@ -508,14 +547,6 @@ struct FPakCommandLineParameters
 	bool bAsyncCompression;
 	bool bAlignFilesLargerThanBlock;	// Align files that are larger than block size
 	bool bForceCompress; // Force all files that request compression to be compressed, even if that results in a larger file size
-	bool bAllowForIndexUnload;
-	FString DataDrivenPlatformName;
-};
-
-struct FPakEntryPair
-{
-	FString Filename;
-	FPakEntry Info;
 };
 
 struct FPakInputPair
@@ -1047,18 +1078,18 @@ void ProcessCommandLine(const TCHAR* CmdLine, const TArray<FString>& NonOptionAr
 	else
 	{
 		CmdLineParameters.FileSystemBlockSize = 0;
-}
+	}
 
 	FString CompBlockSizeString;
 	if (FParse::Value(CmdLine, TEXT("-compressionblocksize="), CompBlockSizeString) &&
 		FParse::Value(CmdLine, TEXT("-compressionblocksize="), CmdLineParameters.CompressionBlockSize))
-{
-		if (CompBlockSizeString.EndsWith(TEXT("MB")))
 	{
+		if (CompBlockSizeString.EndsWith(TEXT("MB")))
+		{
 			CmdLineParameters.CompressionBlockSize *= 1024 * 1024;
-	}
+		}
 		else if (CompBlockSizeString.EndsWith(TEXT("KB")))
-	{
+		{
 			CmdLineParameters.CompressionBlockSize *= 1024;
 		}
 	}
@@ -1105,7 +1136,7 @@ void ProcessCommandLine(const TCHAR* CmdLine, const TArray<FString>& NonOptionAr
 			FName FormatName = *Format;
 
 			if (FCompression::IsFormatValid(FormatName))
-	{
+			{
 				CmdLineParameters.CompressionFormats.Add(FormatName);
 				break;
 			}
@@ -1117,10 +1148,8 @@ void ProcessCommandLine(const TCHAR* CmdLine, const TArray<FString>& NonOptionAr
 
 	if (FParse::Value(CmdLine, TEXT("-create="), ResponseFile))
 	{
-		CmdLineParameters.bAllowForIndexUnload = FParse::Param(CmdLine, TEXT("allowforindexunload"));
 		CmdLineParameters.GeneratePatch = FParse::Value(CmdLine, TEXT("-generatepatch="), CmdLineParameters.SourcePatchPakFilename);
 		FParse::Value(CmdLine, TEXT("-outputchangedfiles="), CmdLineParameters.ChangedFilesOutputFilename);
-		FParse::Value(CmdLine, TEXT("-platform="), CmdLineParameters.DataDrivenPlatformName);
 
 		bool bCompress = FParse::Param(CmdLine, TEXT("compress"));
 		bool bEncrypt = FParse::Param(CmdLine, TEXT("encrypt"));
@@ -1822,6 +1851,105 @@ FArchive* CreatePakWriter(const TCHAR* Filename, const FKeyChain& InKeyChain, bo
 	return Writer;
 }
 
+/* Helper for index creation in CreatePakFile.  Used to (conditionally) write Bytes into a SecondaryIndex and serialize into the PrimaryIndex the offset of the SecondaryIndex */
+struct FSecondaryIndexWriter
+{
+private:
+	TArray<uint8>& SecondaryIndexData;
+	FMemoryWriter SecondaryWriter;
+	TArray<uint8>& PrimaryIndexData;
+	FMemoryWriter& PrimaryIndexWriter;
+	FSHAHash SecondaryHash;
+	int64 OffsetToDataInPrimaryIndex = INDEX_NONE;
+	int64 OffsetToSecondaryInPakFile = INDEX_NONE;
+	int64 SecondarySize = 0;
+	bool bShouldWrite;
+
+public:
+	FSecondaryIndexWriter(TArray<uint8>& InSecondaryIndexData, bool bInShouldWrite, TArray<uint8>& InPrimaryIndexData, FMemoryWriter& InPrimaryIndexWriter)
+		: SecondaryIndexData(InSecondaryIndexData)
+		, SecondaryWriter(SecondaryIndexData)
+		, PrimaryIndexData(InPrimaryIndexData)
+		, PrimaryIndexWriter(InPrimaryIndexWriter)
+		, bShouldWrite(bInShouldWrite)
+	{
+		if (bShouldWrite)
+		{
+			SecondaryWriter.SetByteSwapping(PrimaryIndexWriter.ForceByteSwapping());
+		}
+	}
+
+	/**
+	 * Write the condition flag and the Offset,Size,Hash into the primary index.  Offset of the Secondary index cannot be calculated until the PrimaryIndex is done writing later,
+	 * so placeholders are left instead, with a marker to rewrite them later.
+	 */
+	void WritePlaceholderToPrimary()
+	{
+		PrimaryIndexWriter << bShouldWrite;
+		if (bShouldWrite)
+		{
+			OffsetToDataInPrimaryIndex = PrimaryIndexData.Num();
+			PrimaryIndexWriter << OffsetToSecondaryInPakFile;
+			PrimaryIndexWriter << SecondarySize;
+			PrimaryIndexWriter << SecondaryHash;
+		}
+	}
+
+	FMemoryWriter& GetSecondaryWriter()
+	{
+		return SecondaryWriter;
+	}
+
+	/** The caller has populated this->SecondaryIndexData using GetSecondaryWriter().  We now calculate the size, encrypt, and hash, and store the Offset,Size,Hash in the PrimaryIndex */
+	void FinalizeAndRecordOffset(int64 OffsetInPakFile, const TFunction<void(TArray<uint8> & IndexData, FSHAHash & OutHash)>& FinalizeIndexBlock)
+	{
+		if (!bShouldWrite)
+		{
+			return;
+		}
+
+		FinalizeIndexBlock(SecondaryIndexData, SecondaryHash);
+		SecondarySize = SecondaryIndexData.Num();
+		OffsetToSecondaryInPakFile = OffsetInPakFile;
+
+		PrimaryIndexWriter.Seek(OffsetToDataInPrimaryIndex);
+		PrimaryIndexWriter << OffsetToSecondaryInPakFile;
+		PrimaryIndexWriter << SecondarySize;
+		PrimaryIndexWriter << SecondaryHash;
+	}
+};
+
+/* Verify that Indexes constructed for serialization into the PakFile match the originally collected list of FPakEntryPairs */
+void VerifyIndexesMatch(TArray<FPakEntryPair>& EntryList, FPakFile::FDirectoryIndex& DirectoryIndex, FPakFile::FPathHashIndex& PathHashIndex, uint64 PathHashSeed, const FString& MountPoint,
+	const TArray<uint8>& EncodedPakEntries, const TArray<FPakEntry>& NonEncodableEntries, int32 NumEncodedEntries, int32 NumDeletedEntries, FPakInfo& Info)
+{
+	check(NumEncodedEntries + NonEncodableEntries.Num() + NumDeletedEntries == EntryList.Num());
+
+	FPakEntry EncodedEntry;
+	for (FPakEntryPair& Pair : EntryList)
+	{
+		FString FullPath = FPaths::Combine(MountPoint, Pair.Filename);
+
+		const FPakEntryLocation* PakEntryLocation = FPakFile::FindLocationFromIndex(FullPath, MountPoint, DirectoryIndex);
+		if (!PakEntryLocation)
+		{
+			check(false);
+			continue;
+		}
+		const FPakEntryLocation* PathHashLocation = FPakFile::FindLocationFromIndex(FullPath, MountPoint, PathHashIndex, PathHashSeed);
+		if (!PathHashLocation)
+		{
+			check(false);
+			continue;
+		}
+		check(*PakEntryLocation == *PathHashLocation);
+
+		check(FPakFile::GetPakEntry(*PakEntryLocation, &EncodedEntry, EncodedPakEntries, NonEncodableEntries, Info) != FPakFile::EFindResult::NotFound);
+		check(Pair.Info.IsDeleteRecord() == EncodedEntry.IsDeleteRecord());
+		check(Pair.Info.IsDeleteRecord() || Pair.Info.IndexDataEquals(EncodedEntry));
+	}
+};
+
 TAtomic<int64> GTotalFilesWithPoorForcedCompression(0);
 TAtomic<int64> GTotalExtraMemoryForPoorForcedCompression(0);
 
@@ -2061,15 +2189,15 @@ bool CreatePakFile(const TCHAR* Filename, TArray<FPakInputPair>& FilesToAdd, con
 
 	for (int32 FileIndex = 0; FileIndex < FilesToAdd.Num(); FileIndex++)
 	{
-			AsyncCompressors[FileIndex].Init(&FilesToAdd[FileIndex], CmdLineParameters, &NoPluginCompressionExtensions);
-			if (bRunAsync)
-			{
+		AsyncCompressors[FileIndex].Init(&FilesToAdd[FileIndex], CmdLineParameters, &NoPluginCompressionExtensions);
+		if (bRunAsync)
+		{
 			if (FilesToAdd[FileIndex].bNeedsCompression)
 			{
 				(new FAutoDeleteAsyncTask<FRunCompressionTask>(&AsyncCompressors[FileIndex]))->StartBackgroundTask();
 			}
-		else
-		{
+			else
+			{
 				// call compress function inline 
 				// it won't do anything except for initialize some internal variables used in the non compressed path
 				// we don't want to pass these to a different thread as they may cause congestion with legitimate tasks
@@ -2347,174 +2475,173 @@ bool CreatePakFile(const TCHAR* Filename, TArray<FPakInputPair>& FilesToAdd, con
 
 	FThreadLocalScratchSpace::Get().CleanUp();
 
-	// Remember IndexOffset
-	Info.IndexOffset = PakFileHandle->Tell();
-
-	FPakFileData Data;
-	FPakFile::MakeDirectoryFromPath(MountPoint);
-	Data.MountPoint = MountPoint;
+	auto FinalizeIndexBlockSize = [&Info](TArray<uint8>& IndexData)
+	{
+		if (Info.bEncryptedIndex)
+		{
+			int32 OriginalSize = IndexData.Num();
+			int32 AlignedSize = Align(OriginalSize, FAES::AESBlockSize);
 
-	// process each entry similar to how loading was working
-	for (int32 EntryIndex = 0; EntryIndex < Index.Num(); EntryIndex++)
+			for (int32 PaddingIndex = IndexData.Num(); PaddingIndex < AlignedSize; ++PaddingIndex)
+			{
+				uint8 Byte = IndexData[PaddingIndex % OriginalSize];
+				IndexData.Add(Byte);
+			}
+		}
+	};
+	auto FinalizeIndexBlock = [&TotalEncryptedDataSize, &InKeyChain, &Info, PakHandle = PakFileHandle.Get(), &FinalizeIndexBlockSize] (TArray<uint8>& IndexData, FSHAHash& OutHash)
 	{
-		FPakEntryPair& Entry = Index[EntryIndex];
-		// add raw entry
-		Data.Files.Add(Entry.Info);
+		FinalizeIndexBlockSize(IndexData);
+
+		FSHA1::HashBuffer(IndexData.GetData(), IndexData.Num(), OutHash.Hash);
 
-		// make index entry
-		FString Path = FPaths::GetPath(Entry.Filename);
-		FPakFile::MakeDirectoryFromPath(Path);
-		FPakDirectory* Directory = Data.Index.Find(Path);
-		if (Directory != nullptr)
+		if (Info.bEncryptedIndex)
 		{
-			Directory->Add(FPaths::GetCleanFilename(Entry.Filename), EntryIndex);
+			check(InKeyChain.MasterEncryptionKey);
+			FAES::EncryptData(IndexData.GetData(), IndexData.Num(), InKeyChain.MasterEncryptionKey->Key);
+			TotalEncryptedDataSize += IndexData.Num();
 		}
-		else
-		{
-			FPakDirectory& NewDirectory = Data.Index.Add(Path);
-			NewDirectory.Add(FPaths::GetCleanFilename(Entry.Filename), EntryIndex);
+	};
 
-			// add the parent directories up to the mount point
-			while (MountPoint != Path)
+	if (RequiredPatchPadding > 0)
+	{
+		for (const FPakEntryPair& Pair : Index)
+		{
+			const FString& EntryFilename = Pair.Filename;
+			const FPakEntry& PakEntry = Pair.Info;
+			int64 EntrySize = PakEntry.GetSerializedSize(FPakInfo::PakFile_Version_Latest);
+			int64 TotalSizeToWrite = PakEntry.Size + EntrySize;
+			if (TotalSizeToWrite >= RequiredPatchPadding)
 			{
-				Path = Path.Left(Path.Len() - 1);
-				int32 Offset = 0;
-				if (Path.FindLastChar('/', Offset))
+				int64 RealStart = PakEntry.Offset;
+				if ((RealStart % RequiredPatchPadding) != 0 &&
+					!EntryFilename.EndsWith(TEXT("uexp")) && // these are export sections of larger files and may be packed with uasset/umap and so we don't need a warning here
+					!(EntryFilename.EndsWith(TEXT(".m.ubulk")) && CmdLineParameters.AlignForMemoryMapping > 0)) // Bulk padding unaligns patch padding and so we don't need a warning here
 				{
-					Path = Path.Left(Offset);
-					FPakFile::MakeDirectoryFromPath(Path);
-					if (Data.Index.Find(Path) == nullptr)
-					{
-						Data.Index.Add(Path);
-					}
-				}
-				else
-				{
-					Path = MountPoint;
+					UE_LOG(LogPakFile, Warning, TEXT("File at offset %lld of size %lld not aligned to patch size %i"), RealStart, PakEntry.Size, RequiredPatchPadding);
 				}
 			}
 		}
 	}
 
-	// look up platform info via DataDrivenPlatformInfo.ini (null here means that there was no platform specified, and we cannot freeze)
-	const FDataDrivenPlatformInfoRegistry::FPlatformInfo* PlatformInfo = nullptr;
-	if (CmdLineParameters.DataDrivenPlatformName.Len() > 0)
-	{
-		const TMap<FString, FDataDrivenPlatformInfoRegistry::FPlatformInfo>& Infos = FDataDrivenPlatformInfoRegistry::GetAllPlatformInfos();
-		PlatformInfo = Infos.Find(CmdLineParameters.DataDrivenPlatformName);
-		if (PlatformInfo == nullptr)
-		{
-			UE_LOG(LogPakFile, Fatal, TEXT("Unable to find DataDrivenPlatform '%s', unable to continue safely"), *CmdLineParameters.DataDrivenPlatformName);
-		}
-	}
+	// Create the second copy of the FPakEntries stored in the Index at the end of the PakFile
+	// FPakEntries in the Index are stored as compacted bytes in EncodedPakEntries if possible, or in uncompacted NonEncodableEntries if not.
+	// At the same time, create the two Indexes that map to the encoded-or-not given FPakEntry.  The runtime will only load one of these, depending
+	// on whether it needs to be able to do DirectorySearches.
+	TArray<uint8> EncodedPakEntries;
+	int32 NumEncodedEntries = 0;
+	int32 NumDeletedEntries = 0;
+	TArray<FPakEntry> NonEncodableEntries;
 
-	// check to make sure freezing of the index is okay (target specified, it's 64-bit, and we don't want to modify it at runtime)
-	bool bAllowFreezing = PlatformInfo != nullptr && PlatformInfo->Freezing_b32Bit == false && CmdLineParameters.bAllowForIndexUnload == false;
-
-	// Disable encrypted freezing until frozen decryption is implemented in FPakFile::LoadIndex()
-	bAllowFreezing = bAllowFreezing && !Info.bEncryptedIndex;
+	FPakFile::FDirectoryIndex DirectoryIndex;
+	FPakFile::FPathHashIndex PathHashIndex;
+	TMap<uint64, FString> CollisionDetection; // Currently detecting Collisions only within the files stored into a single Pak.  TODO: Create separate job to detect collisions over all files in the export.
+	uint64 PathHashSeed;
+	int32 NextIndex = 0;
+	auto ReadNextEntry = [&Index, &NextIndex]() -> FPakEntryPair&
+	{
+		return Index[NextIndex++];
+	};
 
-	// if the want to allow for unloading of indices, then 
-	if (bAllowFreezing)
+	FPakFile::EncodePakEntriesIntoIndex(Index.Num(), ReadNextEntry, Filename, Info, MountPoint, NumEncodedEntries, NumDeletedEntries, &PathHashSeed, &DirectoryIndex, &PathHashIndex, EncodedPakEntries, NonEncodableEntries, &CollisionDetection);
+	VerifyIndexesMatch(Index, DirectoryIndex, PathHashIndex, PathHashSeed, MountPoint, EncodedPakEntries, NonEncodableEntries, NumEncodedEntries, NumDeletedEntries, Info);
+
+	// We write one PrimaryIndex and two SecondaryIndexes to the Pak File
+	// PrimaryIndex
+	//		Common scalar data such as MountPoint
+	//		PresenceBit and Offset,Size,Hash for the SecondaryIndexes
+	//		PakEntries (Encoded and NonEncoded)
+	// SecondaryIndex PathHashIndex: used by default in shipped versions of games.  Uses less memory, but does not provide access to all filenames.
+	//		TMap from hash of FilePath to FPakEntryLocation
+	//		Pruned DirectoryIndex, containing only the FilePaths that were requested kept by whitelist config variables
+	// SecondaryIndex FullDirectoryIndex: used for developer tools and for titles that opt out of PathHashIndex because they need access to all filenames.
+	//		TMap from DirectoryPath to FDirectory, which itself is a TMap from CleanFileName to FPakEntryLocation
+	// Each Index is separately encrypted and hashed.  Runtime consumer such as the tools or the client game will only load one of these off of disk (unless runtime verification is turned on).
+
+	// Create the pruned DirectoryIndex for use in the Primary Index
+	TMap<FString, FPakDirectory> PrunedDirectoryIndex;
+	FPakFile::PruneDirectoryIndex(DirectoryIndex, &PrunedDirectoryIndex, MountPoint);
+
+	bool bWritePathHashIndex = FPakFile::IsPakWritePathHashIndex();
+	bool bWriteFullDirectoryIndex = FPakFile::IsPakWriteFullDirectoryIndex();
+	checkf(bWritePathHashIndex || bWriteFullDirectoryIndex, TEXT("At least one of Engine:[Pak]:WritePathHashIndex and Engine:[Pak]:WriteFullDirectoryIndex must be true"));
+
+	TArray<uint8> PrimaryIndexData;
+	TArray<uint8> PathHashIndexData;
+	TArray<uint8> FullDirectoryIndexData;
+	Info.IndexOffset = PakFileHandle->Tell();
+	// Write PrimaryIndex bytes
 	{
-		UE_LOG(LogPakFile, Display, TEXT("ENABLING pak file index freezing"));
+		FMemoryWriter PrimaryIndexWriter(PrimaryIndexData);
+		PrimaryIndexWriter.SetByteSwapping(PakFileHandle->ForceByteSwapping());
 
-		FMemoryImage MemoryImage;
-		// we don't have any editor only data while freezing, so pass false
-		MemoryImage.TargetLayoutParameters.InitializeForPlatform(CmdLineParameters.DataDrivenPlatformName, false);
+		PrimaryIndexWriter << MountPoint;
+		int32 NumEntries = Index.Num();
+		PrimaryIndexWriter << NumEntries;
+		PrimaryIndexWriter << PathHashSeed;
 
-		FMemoryImageWriter Writer(MemoryImage);
-		Writer.WriteObject(Data);
+		FSecondaryIndexWriter PathHashIndexWriter(PathHashIndexData, bWritePathHashIndex, PrimaryIndexData, PrimaryIndexWriter);
+		PathHashIndexWriter.WritePlaceholderToPrimary();
 
-		FMemoryImageResult Result;
-		MemoryImage.Flatten(Result);
+		FSecondaryIndexWriter FullDirectoryIndexWriter(FullDirectoryIndexData, bWriteFullDirectoryIndex, PrimaryIndexData, PrimaryIndexWriter);
+		FullDirectoryIndexWriter.WritePlaceholderToPrimary();
 
-		if (Info.bEncryptedIndex)
+		PrimaryIndexWriter << EncodedPakEntries;
+		int32 NonEncodableEntriesNum = NonEncodableEntries.Num();
+		PrimaryIndexWriter << NonEncodableEntriesNum;
+		for (FPakEntry& PakEntry : NonEncodableEntries)
 		{
-			check(InKeyChain.MasterEncryptionKey);
-			Result.Bytes.AddZeroed(Align(Result.Bytes.Num(), FAES::AESBlockSize) - Result.Bytes.Num());
-			FSHA1::HashBuffer(Result.Bytes.GetData(), Result.Bytes.Num(), Info.IndexHash.Hash);
-			FAES::EncryptData(Result.Bytes.GetData(), Result.Bytes.Num(), InKeyChain.MasterEncryptionKey->Key);
-			TotalEncryptedDataSize += Result.Bytes.Num();
+			PakEntry.Serialize(PrimaryIndexWriter, FPakInfo::PakFile_Version_Latest);
 		}
 
-		int32 Size = Result.Bytes.Num();
-		PakFileHandle->Serialize(Result.Bytes.GetData(), Size);
-		Result.SaveToArchive(*PakFileHandle);
+		// Finalize the size of the PrimaryIndex (it may change due to alignment padding) because we need the size to know the offset of the SecondaryIndexes which come after it in the PakFile.
+		// Do not encrypt and hash it yet, because we still need to replace placeholder data in it for the Offset,Size,Hash of each SecondaryIndex
+		FinalizeIndexBlockSize(PrimaryIndexData);
 
-		// use this to store the frozen data size
-		Info.IndexSize = Size;
-		Info.bIndexIsFrozen = true;
-	}
-	else
-	{
-		UE_LOG(LogPakFile, Display, TEXT("DISABLING pak file index freezing (all must be true: HasPlatformInfo? %s - TargetIs64Bit? %s - NoRuntimeUnloading? %s - Unencrypted? %s)"),
-			PlatformInfo == nullptr ? TEXT("false") : TEXT("true"),
-			PlatformInfo == nullptr ? TEXT("unknown") : PlatformInfo->Freezing_b32Bit ? TEXT("false") : TEXT("true"),
-			CmdLineParameters.bAllowForIndexUnload ? TEXT("false") : TEXT("true"),
-			Info.bEncryptedIndex ? TEXT("false") : TEXT("true"));
-
-		// Serialize Pak Index at the end of Pak File
-		TArray<uint8> IndexData;
-		FMemoryWriter IndexWriter(IndexData);
-		IndexWriter.SetByteSwapping(PakFileHandle->ForceByteSwapping());
-		int32 NumEntries = Index.Num();
-		IndexWriter << MountPoint;
-		IndexWriter << NumEntries;
-		for (int32 EntryIndex = 0; EntryIndex < Index.Num(); EntryIndex++)
+		// Write PathHashIndex bytes
+		if (bWritePathHashIndex)
 		{
-			FPakEntryPair& Entry = Index[EntryIndex];
-			IndexWriter << Entry.Filename;
-			Entry.Info.Serialize(IndexWriter, Info.Version);
-
-			// @todo loadtime: can't this be done when writing the fileentry itself?
-			if (RequiredPatchPadding > 0)
 			{
-				int64 EntrySize = Entry.Info.GetSerializedSize(FPakInfo::PakFile_Version_Latest);
-				int64 TotalSizeToWrite = Entry.Info.Size + EntrySize;
-				if (TotalSizeToWrite >= RequiredPatchPadding)
-				{
-					int64 RealStart = Entry.Info.Offset;
-					if ((RealStart % RequiredPatchPadding) != 0 &&
-						!Entry.Filename.EndsWith(TEXT("uexp")) && // these are export sections of larger files and may be packed with uasset/umap and so we don't need a warning here
-						!(Entry.Filename.EndsWith(TEXT(".m.ubulk")) && CmdLineParameters.AlignForMemoryMapping > 0)) // Bulk padding unaligns patch padding and so we don't need a warning here
-					{
-						UE_LOG(LogPakFile, Warning, TEXT("File at offset %lld of size %lld not aligned to patch size %i"), RealStart, Entry.Info.Size, RequiredPatchPadding);
-					}
-				}
+				FMemoryWriter& SecondaryWriter = PathHashIndexWriter.GetSecondaryWriter();
+				SecondaryWriter << PathHashIndex;
+				SecondaryWriter << PrunedDirectoryIndex;
 			}
+			PathHashIndexWriter.FinalizeAndRecordOffset(Info.IndexOffset + PrimaryIndexData.Num(), FinalizeIndexBlock);
 		}
 
-		if (Info.bEncryptedIndex)
+		// Write FullDirectoryIndex bytes
+		if (bWriteFullDirectoryIndex)
 		{
-			int32 OriginalSize = IndexData.Num();
-			int32 AlignedSize = Align(OriginalSize, FAES::AESBlockSize);
-
-			for (int32 PaddingIndex = IndexData.Num(); PaddingIndex < AlignedSize; ++PaddingIndex)
 			{
-				uint8 Byte = IndexData[PaddingIndex % OriginalSize];
-				IndexData.Add(Byte);
+				FMemoryWriter& SecondaryWriter = FullDirectoryIndexWriter.GetSecondaryWriter();
+				SecondaryWriter << DirectoryIndex;
 			}
+			FullDirectoryIndexWriter.FinalizeAndRecordOffset(Info.IndexOffset + PrimaryIndexData.Num() + PathHashIndexData.Num(), FinalizeIndexBlock);
 		}
 
-		FSHA1::HashBuffer(IndexData.GetData(), IndexData.Num(), Info.IndexHash.Hash);
-
-		if (Info.bEncryptedIndex)
-		{
-			check(InKeyChain.MasterEncryptionKey);
-			FAES::EncryptData(IndexData.GetData(), IndexData.Num(), InKeyChain.MasterEncryptionKey->Key);
-			TotalEncryptedDataSize += IndexData.Num();
-		}
-
-		PakFileHandle->Serialize(IndexData.GetData(), IndexData.Num());
+		// Encrypt and Hash the PrimaryIndex now that we have filled in the SecondaryIndex information
+		FinalizeIndexBlock(PrimaryIndexData, Info.IndexHash);
+	}
 
-		Info.IndexSize = IndexData.Num();
+	// Write the bytes for each Index into the PakFile
+	Info.IndexSize = PrimaryIndexData.Num();
+	PakFileHandle->Serialize(PrimaryIndexData.GetData(), PrimaryIndexData.Num());
+	if (bWritePathHashIndex)
+	{
+		PakFileHandle->Serialize(PathHashIndexData.GetData(), PathHashIndexData.Num());
+	}
+	if (bWriteFullDirectoryIndex)
+	{
+		PakFileHandle->Serialize(FullDirectoryIndexData.GetData(), FullDirectoryIndexData.Num());
 	}
 
-	// Save trailer (offset, size, hash value)
+	// Save the FPakInfo, which has offset, size, and hash value for the PrimaryIndex, at the end of the PakFile
 	Info.Serialize(*PakFileHandle, FPakInfo::PakFile_Version_Latest);
 
 	UE_LOG(LogPakFile, Display, TEXT("Added %d files, %lld bytes total, time %.2lfs."), Index.Num(), PakFileHandle->TotalSize(), FPlatformTime::Seconds() - StartTime);
+	UE_LOG(LogPakFile, Display, TEXT("PrimaryIndex size: %d bytes"), PrimaryIndexData.Num());
+	UE_LOG(LogPakFile, Display, TEXT("PathHashIndex size: %d bytes"), PathHashIndexData.Num());
+	UE_LOG(LogPakFile, Display, TEXT("FullDirectoryIndex size: %d bytes"), FullDirectoryIndexData.Num());
 	if (TotalUncompressedSize)
 	{
 		float PercentLess = ((float)TotalCompressedSize / (TotalUncompressedSize / 100.f));
@@ -2575,12 +2702,12 @@ bool CreatePakFile(const TCHAR* Filename, TArray<FPakInputPair>& FilesToAdd, con
 	return true;
 }
 
-bool TestPakFile(const TCHAR* Filename)
+bool TestPakFile(const TCHAR* Filename, bool TestHashes)
 {	
 	FPakFile PakFile(&FPlatformFileManager::Get().GetPlatformFile(), Filename, false);
 	if (PakFile.IsValid())
 	{
-		return PakFile.Check();
+		return TestHashes ? PakFile.Check() : true;
 	}
 	else
 	{
@@ -2591,7 +2718,8 @@ bool TestPakFile(const TCHAR* Filename)
 
 bool ListFilesInPak(const TCHAR * InPakFilename, int64 SizeFilter, bool bIncludeDeleted, const FString& CSVFilename, bool bExtractToMountPoint, const FKeyChain& InKeyChain)
 {
-	FPakFile PakFile(&FPlatformFileManager::Get().GetPlatformFile(), InPakFilename, false);
+	IPlatformFile* LowerLevelPlatformFile = &FPlatformFileManager::Get().GetPlatformFile();
+	FPakFile PakFile(LowerLevelPlatformFile, InPakFilename, false);
 	int32 FileCount = 0;
 	int64 FileSize = 0;
 	int64 FilteredSize = 0;
@@ -2600,16 +2728,16 @@ bool ListFilesInPak(const TCHAR * InPakFilename, int64 SizeFilter, bool bInclude
 	{
 		UE_LOG(LogPakFile, Display, TEXT("Mount point %s"), *PakFile.GetMountPoint());
 
-		TArray<FPakFile::FFileIterator> Records;
+		TArray<FPakFile::FPakEntryIterator> Records;
 
-		for (FPakFile::FFileIterator It(PakFile,bIncludeDeleted); It; ++It)
+		for (FPakFile::FPakEntryIterator It(PakFile,bIncludeDeleted); It; ++It)
 		{
 			Records.Add(It);
 		}
 
 		struct FOffsetSort
 		{
-			FORCEINLINE bool operator()(const FPakFile::FFileIterator& A, const FPakFile::FFileIterator& B) const
+			FORCEINLINE bool operator()(const FPakFile::FPakEntryIterator& A, const FPakFile::FPakEntryIterator& B) const
 			{
 				return A.Info().Offset < B.Info().Offset;
 			}
@@ -2618,24 +2746,37 @@ bool ListFilesInPak(const TCHAR * InPakFilename, int64 SizeFilter, bool bInclude
 		Records.Sort(FOffsetSort());
 
 		const FString MountPoint = bExtractToMountPoint ? PakFile.GetMountPoint() : TEXT("");
+		FileCount = Records.Num();
+
+		// The Hashes are not stored in the FPakEntries stored in the index, but they are stored for the FPakEntries stored before each payload.
+		// Read the hashes out of the payload
+		TArray<uint8> HashesBuffer;
+		HashesBuffer.SetNum(FileCount * sizeof(FPakEntry::Hash));
+		uint8* Hashes = HashesBuffer.GetData();
+		int32 EntryIndex = 0;
+		for (auto It : Records)
+		{
+			PakFile.ReadHashFromPayload(It.Info(), Hashes + (EntryIndex++)*sizeof(FPakEntry::Hash));
+		}
 
 		if (CSVFilename.Len() > 0)
 		{
-			
 			TArray<FString> Lines;
 			Lines.Empty(Records.Num()+2);
 			Lines.Add(TEXT("Filename, Offset, Size, Hash, Deleted, Compressed, CompressionMethod"));
+			EntryIndex = 0;
 			for (auto It : Records)
 			{
 				const FPakEntry& Entry = It.Info();
+				uint8* Hash = Hashes + (EntryIndex++)*sizeof(FPakEntry::Hash);
 
 				bool bWasCompressed = Entry.CompressionMethodIndex != 0;
 
 				Lines.Add( FString::Printf(
 					TEXT("%s%s, %lld, %lld, %s, %s, %s, %d"),
-					*MountPoint, *It.Filename(),
+					*MountPoint, It.TryGetFilename() ? **It.TryGetFilename() : TEXT("<FileNamesNotLoaded>"),
 					Entry.Offset, Entry.Size,
-					*BytesToHex(Entry.Hash, sizeof(Entry.Hash)),
+					*BytesToHex(Hash, sizeof(FPakEntry::Hash)),
 					Entry.IsDeleteRecord() ? TEXT("true") : TEXT("false"),
 					bWasCompressed ? TEXT("true") : TEXT("false"),
 					Entry.CompressionMethodIndex) );
@@ -2692,9 +2833,11 @@ bool ListFilesInPak(const TCHAR * InPakFilename, int64 SizeFilter, bool bInclude
 				}
 			}
 		}
+		EntryIndex = 0;
 		for (auto It : Records)
 		{
 			const FPakEntry& Entry = It.Info();
+			uint8* Hash = Hashes + (EntryIndex++) * sizeof(FPakEntry::Hash);
 			if (Entry.Size >= SizeFilter)
 			{
 				if (InspectChunkRanges.Num() > 0)
@@ -2706,18 +2849,20 @@ bool ListFilesInPak(const TCHAR * InPakFilename, int64 SizeFilter, bool bInclude
 					{
 						if (InspectChunks.Contains(Chunk))
 						{
-							UE_LOG(LogPakFile, Display, TEXT("[%d - %d] \"%s%s\" offset: %lld, size: %d bytes, sha1: %s, compression: %s."), FirstChunk, LastChunk, *MountPoint, *It.Filename(), Entry.Offset, Entry.Size, *BytesToHex(Entry.Hash, sizeof(Entry.Hash)), *PakFile.GetInfo().GetCompressionMethod(Entry.CompressionMethodIndex).ToString());
+							UE_LOG(LogPakFile, Display, TEXT("[%d - %d] \"%s%s\" offset: %lld, size: %d bytes, sha1: %s, compression: %s."), FirstChunk, LastChunk, *MountPoint,
+								It.TryGetFilename() ? **It.TryGetFilename() : TEXT("<FileNamesNotLoaded"), Entry.Offset, Entry.Size, *BytesToHex(Hash, sizeof(FPakEntry::Hash)),
+								*PakFile.GetInfo().GetCompressionMethod(Entry.CompressionMethodIndex).ToString());
 							break;
 						}
 					}
 				}
 				else
 				{
-					UE_LOG(LogPakFile, Display, TEXT("\"%s%s\" offset: %lld, size: %d bytes, sha1: %s, compression: %s."), *MountPoint, *It.Filename(), Entry.Offset, Entry.Size, *BytesToHex(Entry.Hash, sizeof(Entry.Hash)), *PakFile.GetInfo().GetCompressionMethod(Entry.CompressionMethodIndex).ToString());
+					UE_LOG(LogPakFile, Display, TEXT("\"%s%s\" offset: %lld, size: %d bytes, sha1: %s, compression: %s."), *MountPoint, It.TryGetFilename() ? **It.TryGetFilename() : TEXT("<FileNamesNotLoaded"),
+						Entry.Offset, Entry.Size, *BytesToHex(Hash, sizeof(FPakEntry::Hash)), *PakFile.GetInfo().GetCompressionMethod(Entry.CompressionMethodIndex).ToString());
 				}
 			}
 			FileSize += Entry.Size;
-			FileCount++;
 		}
 		UE_LOG(LogPakFile, Display, TEXT("%d files (%lld bytes), (%lld filtered bytes)."), FileCount, FileSize, FilteredSize);
 
@@ -2830,14 +2975,20 @@ bool AuditPakFiles( const FString& InputPath, bool bOnlyDeleted, const FString&
 			FString PakMountPoint = PakFile.GetMountPoint().Replace(TEXT("../../../"), TEXT(""));
 
 			const bool bIncludeDeleted = true;
-			for (FPakFile::FFileIterator It(PakFile,bIncludeDeleted); It; ++It)
+			if (!PakFile.HasFilenames())
+			{
+				UE_LOG(LogPakFile, Error, TEXT("Pakfiles were loaded without Filenames, cannot audit."));
+				return false;
+			}
+
+			for (FPakFile::FPakEntryIterator It(PakFile,bIncludeDeleted); It; ++It)
 			{
 				FString AssetName = PakMountPoint;
 				if (!AssetName.IsEmpty() && !AssetName.EndsWith("/"))
 				{
 					AssetName += "/";
 				}
-				AssetName += It.Filename();
+				AssetName += *It.TryGetFilename();
 
 				FFilePakRevision Revision;
 				Revision.PakFilename = PakFileList[PakFileIndex];
@@ -3090,7 +3241,7 @@ bool ListFilesAtOffset( const TCHAR* InPakFileName, const TArray<int64>& InOffse
 
 	TArray<int64> OffsetsToCheck = InOffsets;
 	FArchive& PakReader = *PakFile.GetSharedReader(NULL);
-	for (FPakFile::FFileIterator It(PakFile); It; ++It)
+	for (FPakFile::FPakEntryIterator It(PakFile); It; ++It)
 	{
 		const FPakEntry& Entry = It.Info();
 
@@ -3100,7 +3251,8 @@ bool ListFilesAtOffset( const TCHAR* InPakFileName, const TArray<int64>& InOffse
 		{
 			if( Offset >= Entry.Offset && Offset <= Entry.Offset+Entry.Size )
 			{
-				UE_LOG( LogPakFile, Display, TEXT("%-12lld%-12lld%-12d%s"), Offset, Entry.Offset, Entry.Size, *It.Filename() );
+				const FString* Filename = It.TryGetFilename();
+				UE_LOG( LogPakFile, Display, TEXT("%-12lld%-12lld%-12d%s"), Offset, Entry.Offset, Entry.Size, Filename ? **Filename : TEXT("<FileNamesNotLoaded>") );
 				FoundOffset = Offset;
 				break;
 			}
@@ -3147,13 +3299,13 @@ bool ShowCompressionBlockCRCs( const TCHAR* InPakFileName, TArray<int64>& InOffs
 
 		//find the matching entry
 		const FPakEntry* Entry = nullptr;
-		for (FPakFile::FFileIterator It(PakFile); It; ++It)
+		for (FPakFile::FPakEntryIterator It(PakFile); It; ++It)
 		{
 			const FPakEntry& ThisEntry = It.Info();
 			if( Offset >= ThisEntry.Offset && Offset <= ThisEntry.Offset+ThisEntry.Size )
 			{
 				Entry = &ThisEntry;
-				FString EntryFilename = It.Filename();
+				FString EntryFilename = It.TryGetFilename() ? *It.TryGetFilename() : TEXT("<FileNamesNotLoaded>");
 				FName EntryCompressionMethod = PakFile.GetInfo().GetCompressionMethod(Entry->CompressionMethodIndex);
 
 				UE_LOG(LogPakFile, Display, TEXT("Offset: %lld  -> EntrySize: %lld  Encrypted: %-3s  Compression: %-8s  [%s]"), Offset, Entry->Size, Entry->IsEncrypted() ? TEXT("Yes") : TEXT("No"), *EntryCompressionMethod.ToString(), *EntryFilename );
@@ -3272,11 +3424,17 @@ bool GeneratePIXMappingFile(const TArray<FString> InPakFileList, const FString&
 
 		const FString PakFileMountPoint = PakFile.GetMountPoint();
 		FArchive& PakReader = *PakFile.GetSharedReader(NULL);
-		for (FPakFile::FFileIterator It(PakFile); It; ++It)
+		if (!PakFile.HasFilenames())
+		{
+			UE_LOG(LogPakFile, Error, TEXT("PakFiles were loaded without filenames, cannot generate PIX mapping file."));
+			return false;
+		}
+
+		for (FPakFile::FPakEntryIterator It(PakFile); It; ++It)
 		{
 			const FPakEntry& Entry = It.Info();
 
-			CSVFileWriter->Logf(TEXT("0x%010llx,0x%08llx,%s"), Entry.Offset, Entry.Size, *(PakFileMountPoint / It.Filename()));
+			CSVFileWriter->Logf(TEXT("0x%010llx,0x%08llx,%s"), Entry.Offset, Entry.Size, *(PakFileMountPoint / *It.TryGetFilename()));
 		}
 
 		CSVFileWriter->Close();
@@ -3311,6 +3469,17 @@ bool ExtractFilesFromPak(const TCHAR* InPakFilename, TMap<FString, FFileInfo>& I
 		return false;
 	}
 
+	if (OutOrderMap)
+	{
+		if (PakFileList.Num() > 1)
+		{
+			UE_LOG(LogPakFile, Error, TEXT("ExtractFilesFromPak: Collecting the order of the files in a pak is not implemented when multiple packs are specified by wildcard. ")
+				TEXT("ExtractFilesFromPak was called on '%s', which matched %d pakfiles."), InPakFilename, PakFileList.Num());
+			return false;
+		}
+		OutOrderMap->Empty();
+	}
+
 
 	bool bIncludeDeleted = (OutDeletedEntries != nullptr);
 
@@ -3347,29 +3516,35 @@ bool ExtractFilesFromPak(const TCHAR* InPakFilename, TMap<FString, FFileInfo>& I
 			}
 
 			FString PakMountPoint = bUseMountPoint ? PakFile.GetMountPoint().Replace(TEXT("../../../"), TEXT("")) : TEXT("");
+			if (!PakFile.HasFilenames())
+			{
+				UE_LOG(LogPakFile, Error, TEXT("PakFiles were loaded without filenames, cannot extract."));
+				return false;
+			}
 
-			for (FPakFile::FFileIterator It(PakFile,bIncludeDeleted); It; ++It, ++FileCount)
+			for (FPakFile::FPakEntryIterator It(PakFile,bIncludeDeleted); It; ++It, ++FileCount)
 			{
 				// Extract only the most recent version of a file when present in multiple paks
-				FFileInfo* HashFileInfo = InFileHashes.Find(It.Filename());
+				const FString& EntryFileName = *It.TryGetFilename();
+				FFileInfo* HashFileInfo = InFileHashes.Find(EntryFileName);
 				if (HashFileInfo == nullptr || HashFileInfo->PatchIndex == PakPriority)
 				{
-					FString DestFilename(DestPath / PakMountPoint / It.Filename());
+					FString DestFilename(DestPath / PakMountPoint / EntryFileName);
 
 					const FPakEntry& Entry = It.Info();
 					if (Entry.IsDeleteRecord())
 					{
-						UE_LOG(LogPakFile, Display, TEXT("Found delete record for \"%s\"."), *It.Filename() );
+						UE_LOG(LogPakFile, Display, TEXT("Found delete record for \"%s\"."), *EntryFileName);
 
 						FPakInputPair DeleteRecord;
 						DeleteRecord.bIsDeleteRecord = true;
 						DeleteRecord.Source = DestFilename;
-						DeleteRecord.Dest = PakFile.GetMountPoint() / It.Filename();
+						DeleteRecord.Dest = PakFile.GetMountPoint() / EntryFileName;
 						OutDeletedEntries->Add(DeleteRecord);
 						continue;
 					}
 
-					if (InFilter && (!It.Filename().MatchesWildcard(*InFilter)))
+					if (InFilter && (!EntryFileName.MatchesWildcard(*InFilter)))
 					{
 						continue;
 					}
@@ -3378,7 +3553,7 @@ bool ExtractFilesFromPak(const TCHAR* InPakFilename, TMap<FString, FFileInfo>& I
 					uint32 SerializedCrcTest = 0;
 					FPakEntry EntryInfo;
 					EntryInfo.Serialize(PakReader, PakFile.GetInfo().Version);
-					if (EntryInfo == Entry)
+					if (EntryInfo.IndexDataEquals(Entry))
 					{
 						TUniquePtr<FArchive> FileHandle(IFileManager::Get().CreateFileWriter(*DestFilename));
 						if (FileHandle)
@@ -3391,12 +3566,12 @@ bool ExtractFilesFromPak(const TCHAR* InPakFilename, TMap<FString, FFileInfo>& I
 							{
 								UncompressCopyFile(*FileHandle, PakReader, Entry, PersistantCompressionBuffer, CompressionBufferSize, InKeyChain, PakFile);
 							}
-							UE_LOG(LogPakFile, Display, TEXT("Extracted \"%s\" to \"%s\" Offset %d."), *It.Filename(), *DestFilename, Entry.Offset);
+							UE_LOG(LogPakFile, Display, TEXT("Extracted \"%s\" to \"%s\" Offset %d."), *EntryFileName, *DestFilename, Entry.Offset);
 							ExtractedCount++;
 
 							if (OutOrderMap != nullptr)
 							{
-								OutOrderMap->Add(PakFile.GetMountPoint() / It.Filename(), It.GetIndexInPakFile());
+								OutOrderMap->AddOffset(PakFile.GetMountPoint() / EntryFileName, It.Info().Offset);
 							}
 
 							if (OutEntries != nullptr)
@@ -3406,7 +3581,7 @@ bool ExtractFilesFromPak(const TCHAR* InPakFilename, TMap<FString, FFileInfo>& I
 								Input.Source = DestFilename;
 								FPaths::NormalizeFilename(Input.Source);
 
-								Input.Dest = PakFile.GetMountPoint() + FPaths::GetPath(It.Filename());
+								Input.Dest = PakFile.GetMountPoint() + FPaths::GetPath(EntryFileName);
 								FPaths::NormalizeFilename(Input.Dest);
 								FPakFile::MakeDirectoryFromPath(Input.Dest);
 
@@ -3424,7 +3599,7 @@ bool ExtractFilesFromPak(const TCHAR* InPakFilename, TMap<FString, FFileInfo>& I
 					}
 					else
 					{
-						UE_LOG(LogPakFile, Error, TEXT("Serialized hash mismatch for \"%s\"."), *It.Filename());
+						UE_LOG(LogPakFile, Error, TEXT("PakEntry mismatch for \"%s\"."), *EntryFileName);
 						ErrorCount++;
 					}
 				}
@@ -3441,6 +3616,11 @@ bool ExtractFilesFromPak(const TCHAR* InPakFilename, TMap<FString, FFileInfo>& I
 		}
 	}
 
+	if (OutOrderMap)
+	{
+		OutOrderMap->ConvertOffsetsToOrder();
+	}
+
 	return true;
 }
 
@@ -3476,12 +3656,17 @@ bool DiffFilesInPaks(const FString& InPakFilename1, const FString& InPakFilename
 		int64 CompressionBufferSize = 0;
 		uint8* PersistantCompressionBuffer = NULL;
 		int32 ErrorCount = 0;
-		int32 FileCount = 0;		
-		
+		int32 FileCount = 0;
+
+		if (!PakFile1.HasFilenames() || !PakFile2.HasFilenames())
+		{
+			UE_LOG(LogPakFile, Error, TEXT("Pakfiles were loaded without Filenames, cannot diff."));
+			return false;
+		}
 		//loop over pak1 entries.  compare against entry in pak2.
-		for (FPakFile::FFileIterator It(PakFile1); It; ++It, ++FileCount)
+		for (FPakFile::FPakEntryIterator It(PakFile1); It; ++It, ++FileCount)
 		{
-			const FString& PAK1FileName = It.Filename();
+			const FString& PAK1FileName = *It.TryGetFilename();
 
 			//double check entry info and move pakreader into place
 			const FPakEntry& Entry1 = It.Info();
@@ -3560,9 +3745,9 @@ bool DiffFilesInPaks(const FString& InPakFilename1, const FString& InPakFilename
 				}
 			}			
 		}
-		
+
 		//check for files unique to the second pak.
-		for (FPakFile::FFileIterator It(PakFile2); It; ++It, ++FileCount)
+		for (FPakFile::FPakEntryIterator It(PakFile2); It; ++It, ++FileCount)
 		{
 			const FPakEntry& Entry2 = It.Info();
 			PakReader2.Seek(Entry2.Offset);
@@ -3572,7 +3757,7 @@ bool DiffFilesInPaks(const FString& InPakFilename1, const FString& InPakFilename
 
 			if (EntryInfo2 == Entry2)
 			{
-				const FString& PAK2FileName = It.Filename();
+				const FString& PAK2FileName = *It.TryGetFilename();
 				FPakEntry Entry1;
 				FPakFile::EFindResult FoundEntry1 = PakFile1.Find(PakFile2.GetMountPoint() / PAK2FileName, &Entry1);
 				if (FoundEntry1 != FPakFile::EFindResult::Found)
@@ -3665,9 +3850,14 @@ bool GenerateHashesFromPak(const TCHAR* InPakFilename, const TCHAR* InDestPakFil
 			}
 
 			FString PakMountPoint = bUseMountPoint ? PakFile.GetMountPoint().Replace(TEXT("../../../"), TEXT("")) : TEXT("");
+			if (!PakFile.HasFilenames())
+			{
+				UE_LOG(LogPakFile, Error, TEXT("Pakfiles were loaded without Filenames, cannot GenerateHashesFromPak."));
+				return false;
+			}
 
 			const bool bIncludeDeleted = true;
-			for (FPakFile::FFileIterator It(PakFile,bIncludeDeleted); It; ++It, ++FileCount)
+			for (FPakFile::FPakEntryIterator It(PakFile,bIncludeDeleted); It; ++It, ++FileCount)
 			{
 				const FPakEntry& Entry = It.Info();
 				FFileInfo FileHash = {};
@@ -3678,7 +3868,7 @@ bool GenerateHashesFromPak(const TCHAR* InPakFilename, const TCHAR* InDestPakFil
 				{
 					FullFilename += "/";
 				}
-				FullFilename += It.Filename();
+				FullFilename += *It.TryGetFilename();
 
 				if (Entry.IsDeleteRecord())
 				{
@@ -3693,7 +3883,7 @@ bool GenerateHashesFromPak(const TCHAR* InPakFilename, const TCHAR* InDestPakFil
 				    uint32 SerializedCrcTest = 0;
 				    FPakEntry EntryInfo;
 				    EntryInfo.Serialize(PakReader, PakFile.GetInfo().Version);
-				    if (EntryInfo == Entry)
+				    if (EntryInfo.IndexDataEquals(Entry))
 				    {
 					    // TUniquePtr<FArchive> FileHandle(IFileManager::Get().CreateFileWriter(*DestFilename));
 					    TArray<uint8> Bytes;
@@ -3726,7 +3916,7 @@ bool GenerateHashesFromPak(const TCHAR* InPakFilename, const TCHAR* InDestPakFil
 					}
 					else
 					{
-						UE_LOG(LogPakFile, Error, TEXT("Serialized hash mismatch for \"%s\"."), *It.Filename());
+						UE_LOG(LogPakFile, Error, TEXT("Serialized hash mismatch for \"%s\"."), **It.TryGetFilename());
 						ErrorCount++;
 					}
 				}
@@ -4286,14 +4476,19 @@ void ProcessLegacyFileMoves( TArray<FPakInputPair>& InDeleteRecords, TMap<FStrin
 			FString PakMountPoint = PakFile.GetMountPoint().Replace(TEXT("../../../"), TEXT(""));
 
 			const bool bIncludeDeleted = true;
-			for (FPakFile::FFileIterator It(PakFile,bIncludeDeleted); It; ++It)
+			if (!PakFile.HasFilenames())
+			{
+				UE_LOG(LogPakFile, Error, TEXT("Pakfiles were loaded without Filenames, cannot ProcessLegacyFileMoves."));
+				return;
+			}
+			for (FPakFile::FPakEntryIterator It(PakFile,bIncludeDeleted); It; ++It)
 			{
 				FString AssetName = PakMountPoint;
 				if (!AssetName.IsEmpty() && !AssetName.EndsWith("/"))
 				{
 					AssetName += "/";
 				}
-				AssetName += It.Filename();
+				AssetName += *It.TryGetFilename();
 
 				bool bHasNewDeleteRecord = DeleteRecordSourceNames.Contains(AssetName);
 
@@ -4581,7 +4776,9 @@ bool ExecuteUnrealPak(const TCHAR* CmdLine)
 	LoadKeyChain(CmdLine, KeyChain);
 	ApplyEncryptionKeys(KeyChain);
 
-	if (FParse::Param(CmdLine, TEXT("Test")))
+	bool IsTestCommand = FParse::Param(CmdLine, TEXT("Test"));
+	bool IsVerifyCommand = FParse::Param(CmdLine, TEXT("Verify"));
+	if (IsTestCommand || IsVerifyCommand)
 	{
 		if (NonOptionArguments.Num() != 1)
 		{
@@ -4590,55 +4787,7 @@ bool ExecuteUnrealPak(const TCHAR* CmdLine)
 		}
 
 		FString PakFilename = GetPakPath(*NonOptionArguments[0], false);
-		return TestPakFile(*PakFilename);
-	}
-
-	if (FParse::Param(CmdLine, TEXT("TestMemoryOptimization")))
-	{
-		TArray<FPakInputPair> Entries;
-		FPakCommandLineParameters CmdLineParameters;
-		ProcessCommandLine(CmdLine, NonOptionArguments, Entries, CmdLineParameters);
-
-		if (NonOptionArguments.Num() != 1)
-		{
-			UE_LOG(LogPakFile, Error, TEXT("Incorrect arguments. Expected: -TestMemoryOptimization <SourceFolder>"));
-			return false;
-		}
-
-		FString SourceDir = NonOptionArguments[0];
-		TArray<FString> PakFilenames;
-		IFileManager::Get().FindFiles(PakFilenames, *SourceDir, TEXT("*.pak"));
-		TArray<FPakFile*> PakFiles;
-		PakFiles.Empty(PakFilenames.Num());
-
-		for (const FString& PakFilename : PakFilenames)
-		{
-			PakFiles.Add(new FPakFile(&FPlatformFileManager::Get().GetPlatformFile(), *(FPaths::Combine(SourceDir, PakFilename)), false));
-		}
-
-		TMap<uint64, FPakEntry> CollisionChecker;
-		
-		for (FPakFile* PakFile : PakFiles)
-		{
-			if (!PakFile->UnloadPakEntryFilenames(CollisionChecker, nullptr, false))
-			{
-				UE_LOG(LogPakFile, Error, TEXT("Pak '%s' failed to unload filenames"), *PakFile->GetFilename());
-			}
-
-			if (!PakFile->ShrinkPakEntriesMemoryUsage())
-			{
-				UE_LOG(LogPakFile, Error, TEXT("Pak '%s' failed to shrink entries"), *PakFile->GetFilename());
-			}
-		}
-
-		for (FPakFile* PakFile : PakFiles)
-		{
-			delete PakFile;
-		}
-
-		PakFiles.Empty();
-
-		return true;
+		return TestPakFile(*PakFilename, IsVerifyCommand);
 	}
 
 	if (FParse::Param(CmdLine, TEXT("List")))
@@ -4715,7 +4864,12 @@ bool ExecuteUnrealPak(const TCHAR* CmdLine)
 		TArray<FPakInputPair> ResponseContent;
 		TArray<FPakInputPair> DeletedContent;
 		FPakOrderMap OrderMap;
-		if (ExtractFilesFromPak(*PakFilename, EmptyMap, *DestPath, bExtractToMountPoint, KeyChain, bUseFilter ? &Filter : nullptr, &ResponseContent, &DeletedContent, &OrderMap) == false)
+		FPakOrderMap* UsedOrderMap = nullptr;
+		if (bGenerateOrderMap)
+		{
+			UsedOrderMap = &OrderMap;
+		}
+		if (ExtractFilesFromPak(*PakFilename, EmptyMap, *DestPath, bExtractToMountPoint, KeyChain, bUseFilter ? &Filter : nullptr, &ResponseContent, &DeletedContent, UsedOrderMap) == false)
 		{
 			return false;
 		}
@@ -4970,6 +5124,12 @@ bool ExecuteUnrealPak(const TCHAR* CmdLine)
 			return false;
 		}
 
+		if (Entries.Num() == 0)
+		{
+			UE_LOG(LogPakFile, Error, TEXT("No files specified to add to pak file."));
+			return false;
+		}
+
 		int32 LowestSourcePakVersion = 0;
 		TMap<FString, FFileInfo> SourceFileHashes;
 
@@ -5057,6 +5217,7 @@ bool ExecuteUnrealPak(const TCHAR* CmdLine)
 
 	UE_LOG(LogPakFile, Error, TEXT("No pak file name specified. Usage:"));
 	UE_LOG(LogPakFile, Error, TEXT("  UnrealPak <PakFilename> -Test"));
+	UE_LOG(LogPakFile, Error, TEXT("  UnrealPak <PakFilename> -Verify"));
 	UE_LOG(LogPakFile, Error, TEXT("  UnrealPak <PakFilename> -List [-ExcludeDeleted]"));
 	UE_LOG(LogPakFile, Error, TEXT("  UnrealPak <PakFilename> <GameUProjectName> <GameFolderName> -ExportDependencies=<OutputFileBase> -NoAssetRegistryCache -ForceDependsGathering"));
 	UE_LOG(LogPakFile, Error, TEXT("  UnrealPak <PakFilename> -Extract <ExtractDir> [-Filter=<filename>]"));
diff --git a/Engine/Source/Editor/UnrealEd/Private/Commandlets/ExportPakDependencies.cpp b/Engine/Source/Editor/UnrealEd/Private/Commandlets/ExportPakDependencies.cpp
--- a/Engine/Source/Editor/UnrealEd/Private/Commandlets/ExportPakDependencies.cpp
+++ b/Engine/Source/Editor/UnrealEd/Private/Commandlets/ExportPakDependencies.cpp
@@ -127,12 +127,17 @@ bool ExportDependencies(const TCHAR * PakFilename, const TCHAR* GameName, const
 	{
 		// Get size information from PAK file.
 		{
-			TArray<FPakFile::FFileIterator> Records;
+			TArray<FPakFile::FPakEntryIterator> Records;
 			FString PakGameContentFolder = FString(GameName) + TEXT("/Content");
-			for(FPakFile::FFileIterator It(PakFile); It; ++It)
+			if (!PakFile.HasFilenames())
+			{
+				UE_LOG(LogPakFile, Error, TEXT("Pakfiles were loaded without Filenames, cannot export."));
+				return false;
+			}
+			for(FPakFile::FPakEntryIterator It(PakFile); It; ++It)
 			{
 				FString PackageName;
-				It.Filename().Split(TEXT("."),&PackageName,NULL);
+				It.TryGetFilename()->Split(TEXT("."),&PackageName,NULL);
 				int64 Size = It.Info().Size;			
 
 				if( PackageName.StartsWith(TEXT("Engine/Content")) )
diff --git a/Engine/Source/Editor/UnrealEd/Private/Factories/PackFactory.cpp b/Engine/Source/Editor/UnrealEd/Private/Factories/PackFactory.cpp
--- a/Engine/Source/Editor/UnrealEd/Private/Factories/PackFactory.cpp
+++ b/Engine/Source/Editor/UnrealEd/Private/Factories/PackFactory.cpp
@@ -334,7 +334,7 @@ UObject* UPackFactory::FactoryCreateBinary
 
 	UObject* ReturnAsset = nullptr;
 
-	if (PakFile.IsValid())
+	if (PakFile.IsValid() && PakFile.HasFilenames())
 	{
 		static FString ContentFolder(TEXT("/Content/"));
 		FString ContentDestinationRoot = FPaths::ProjectContentDir();
@@ -358,9 +358,11 @@ UObject* UPackFactory::FactoryCreateBinary
 		TArray<FString> WrittenSourceFiles;
 
 		// Process the config files and identify if we have source files
-		for (FPakFile::FFileIterator It(PakFile); It; ++It, ++FileCount)
+		for (FPakFile::FPakEntryIterator It(PakFile); It; ++It, ++FileCount)
 		{
-			if (It.Filename().StartsWith(TEXT("Config/")) || It.Filename().Contains(TEXT("/Config/")))
+			const FString* EntryFilename = It.TryGetFilename();
+			check(EntryFilename);
+			if (EntryFilename->StartsWith(TEXT("Config/")) || EntryFilename->Contains(TEXT("/Config/")))
 			{
 				const FPakEntry& Entry = It.Info();
 				PakReader.Seek(Entry.Offset);
@@ -375,11 +377,11 @@ UObject* UPackFactory::FactoryCreateBinary
 				}
 				else
 				{
-					UE_LOG(LogPackFactory, Error, TEXT("Serialized hash mismatch for \"%s\"."), *It.Filename());
+					UE_LOG(LogPackFactory, Error, TEXT("Serialized hash mismatch for \"%s\"."), **EntryFilename);
 					ErrorCount++;
 				}
 			}
-			else if (!ConfigParameters.bContainsSource && (It.Filename().StartsWith(TEXT("Source/")) || It.Filename().Contains(TEXT("/Source/"))))
+			else if (!ConfigParameters.bContainsSource && (EntryFilename ->StartsWith(TEXT("Source/")) || EntryFilename->Contains(TEXT("/Source/"))))
 			{
 				ConfigParameters.bContainsSource = true;
 			}
@@ -457,16 +459,18 @@ UObject* UPackFactory::FactoryCreateBinary
 		}
 
 		// Process everything else and copy out to disk
-		for (FPakFile::FFileIterator It(PakFile); It; ++It, ++FileCount)
+		for (FPakFile::FPakEntryIterator It(PakFile); It; ++It, ++FileCount)
 		{
+			const FString* EntryFilename = It.TryGetFilename();
+			check(EntryFilename);
 			// config files already handled
-			if (It.Filename().StartsWith(TEXT("Config/")) || It.Filename().Contains(TEXT("/Config/")))
+			if (EntryFilename->StartsWith(TEXT("Config/")) || EntryFilename->Contains(TEXT("/Config/")))
 			{
 				continue;
 			}
 
 			// Media and manifest files don't get written out as part of the install
-			if (It.Filename().Contains(TEXT("manifest.json")) || It.Filename().StartsWith(TEXT("Media/")) || It.Filename().Contains(TEXT("/Media/")))
+			if (EntryFilename->Contains(TEXT("manifest.json")) || EntryFilename->StartsWith(TEXT("Media/")) || EntryFilename->Contains(TEXT("/Media/")))
 			{
 				continue;
 			}
@@ -478,9 +482,9 @@ UObject* UPackFactory::FactoryCreateBinary
 
 			if (EntryInfo == Entry)
 			{
-				if (It.Filename().StartsWith(TEXT("Source/")) || It.Filename().Contains(TEXT("/Source/")))
+				if (EntryFilename->StartsWith(TEXT("Source/")) || EntryFilename->Contains(TEXT("/Source/")))
 				{
-					FString DestFilename = It.Filename();
+					FString DestFilename = *EntryFilename;
 					if (DestFilename.StartsWith(TEXT("Source/")))
 					{
 						DestFilename.RightChopInline(7, false);
@@ -495,7 +499,7 @@ UObject* UPackFactory::FactoryCreateBinary
 					}
 
 					DestFilename = SourceModuleInfo.ModuleSourcePath / DestFilename;
-					UE_LOG(LogPackFactory, Log, TEXT("%s (%ld) -> %s"), *It.Filename(), Entry.Size, *DestFilename);
+					UE_LOG(LogPackFactory, Log, TEXT("%s (%ld) -> %s"), **EntryFilename, Entry.Size, *DestFilename);
 
 					FString SourceContents;
 					PackFactoryHelper::ExtractFileToString(Entry, PakReader, CopyBuffer, PersistentCompressionBuffer, SourceContents, PakFile);
@@ -522,7 +526,7 @@ UObject* UPackFactory::FactoryCreateBinary
 				}
 				else
 				{
-					FString DestFilename = It.Filename();
+					FString DestFilename = *EntryFilename;
 					if (DestFilename.StartsWith(TEXT("Content/")))
 					{
 						DestFilename.RightChopInline(8, false);
@@ -536,7 +540,7 @@ UObject* UPackFactory::FactoryCreateBinary
 						}
 					}
 					DestFilename = ContentDestinationRoot / DestFilename;
-					UE_LOG(LogPackFactory, Log, TEXT("%s (%ld) -> %s"), *It.Filename(), Entry.Size, *DestFilename);
+					UE_LOG(LogPackFactory, Log, TEXT("%s (%ld) -> %s"), **EntryFilename, Entry.Size, *DestFilename);
 
 					TUniquePtr<FArchive> FileHandle(IFileManager::Get().CreateFileWriter(*DestFilename));
 
@@ -554,7 +558,7 @@ UObject* UPackFactory::FactoryCreateBinary
 			}
 			else
 			{
-				UE_LOG(LogPackFactory, Error, TEXT("Serialized hash mismatch for \"%s\"."), *It.Filename());
+				UE_LOG(LogPackFactory, Error, TEXT("Serialized hash mismatch for \"%s\"."), **EntryFilename);
 				ErrorCount++;
 			}
 		}
@@ -762,7 +766,14 @@ UObject* UPackFactory::FactoryCreateBinary
 	}
 	else
 	{
-		UE_LOG(LogPackFactory, Warning, TEXT("Invalid pak file."));
+		if (!PakFile.IsValid())
+		{
+			UE_LOG(LogPackFactory, Warning, TEXT("Invalid pak file."));
+		}
+		else
+		{
+			UE_LOG(LogPakFile, Error, TEXT("Pakfiles were loaded without Filenames, creation aborted."));
+		}
 	}
 
 	return ReturnAsset;
diff --git a/Engine/Source/Runtime/Core/Private/IOS/IOSChunkInstaller.cpp b/Engine/Source/Runtime/Core/Private/IOS/IOSChunkInstaller.cpp
--- a/Engine/Source/Runtime/Core/Private/IOS/IOSChunkInstaller.cpp
+++ b/Engine/Source/Runtime/Core/Private/IOS/IOSChunkInstaller.cpp
@@ -151,11 +151,11 @@ bool FIOSChunkInstall::PrioritizeChunk( uint32 ChunkID, EChunkPriority::Type Pri
             {
                 FIOSChunkStatus* FoundStatus = ChunkStatus.Find(BlockID);
                 NSBundle* Bundle = FoundStatus->Request.bundle;
-                FString PakFile = FString::Printf(TEXT("pakchunk%i-ios.pak"), BlockID);
-                NSString* ResourcePath = [Bundle pathForResource:PakFile.GetNSString() ofType:nil];
+                FString PakFileName = FString::Printf(TEXT("pakchunk%i-ios.pak"), BlockID);
+                NSString* ResourcePath = [Bundle pathForResource: PakFileName.GetNSString() ofType:nil];
                 NSLog(@"ResourcePath: %@", ResourcePath);
-                bool bMounted = FCoreDelegates::OnMountPak.Execute(FString(ResourcePath), 0, nullptr);
-                if (bMounted)
+                IPakFile* PakFile = FCoreDelegates::MountPak.Execute(FString(ResourcePath), 0);
+                if (PakFile)
                 {
                     MountedChunks.Add(BlockID);
 					InstallDelegate.Broadcast(BlockID, true);
diff --git a/Engine/Source/Runtime/Core/Private/Internationalization/TextLocalizationManager.cpp b/Engine/Source/Runtime/Core/Private/Internationalization/TextLocalizationManager.cpp
--- a/Engine/Source/Runtime/Core/Private/Internationalization/TextLocalizationManager.cpp
+++ b/Engine/Source/Runtime/Core/Private/Internationalization/TextLocalizationManager.cpp
@@ -300,7 +300,7 @@ void BeginPreInitTextLocalization()
 	DECLARE_SCOPE_CYCLE_COUNTER(TEXT("BeginPreInitTextLocalization"), STAT_BeginPreInitTextLocalization, STATGROUP_LoadTime);
 
 	// Bind this delegate before the PAK file loader is created
-	FCoreDelegates::OnPakFileMounted.AddRaw(&FTextLocalizationManager::Get(), &FTextLocalizationManager::OnPakFileMounted);
+	FCoreDelegates::OnPakFileMounted2.AddRaw(&FTextLocalizationManager::Get(), &FTextLocalizationManager::OnPakFileMounted);
 }
 
 void BeginInitTextLocalization()
@@ -849,8 +849,9 @@ void FTextLocalizationManager::RefreshResources()
 	LoadLocalizationResourcesForCulture(FInternationalization::Get().GetCurrentLanguage()->GetName(), LocLoadFlags);
 }
 
-void FTextLocalizationManager::OnPakFileMounted(const TCHAR* PakFilename, const int32 ChunkId)
+void FTextLocalizationManager::OnPakFileMounted(const IPakFile& PakFile)
 {
+	int32 ChunkId = PakFile.PakGetPakchunkIndex();
 	if (ChunkId == INDEX_NONE || ChunkId == 0)
 	{
 		// Skip non-chunked PAK files, and chunk 0 as that contains the standard localization data
@@ -860,7 +861,7 @@ void FTextLocalizationManager::OnPakFileMounted(const TCHAR* PakFilename, const
 	// Track this so that full resource refreshes (eg, changing culture) work as expected
 	LocResTextSource->RegisterChunkId(ChunkId);
 
-	UE_LOG(LogTextLocalizationManager, Verbose, TEXT("Request to load localization data for chunk %d (from PAK '%s')"), ChunkId, PakFilename);
+	UE_LOG(LogTextLocalizationManager, Verbose, TEXT("Request to load localization data for chunk %d (from PAK '%s')"), ChunkId, *PakFile.PakGetPakFilename());
 
 	if (!bIsInitialized)
 	{
diff --git a/Engine/Source/Runtime/Core/Private/Misc/CoreDelegates.cpp b/Engine/Source/Runtime/Core/Private/Misc/CoreDelegates.cpp
--- a/Engine/Source/Runtime/Core/Private/Misc/CoreDelegates.cpp
+++ b/Engine/Source/Runtime/Core/Private/Misc/CoreDelegates.cpp
@@ -27,12 +27,16 @@ FCoreDelegates::FOnPreMainInit& FCoreDelegates::GetPreMainInitDelegate()
 }
 
 FCoreDelegates::FOnMountAllPakFiles FCoreDelegates::OnMountAllPakFiles;
+PRAGMA_DISABLE_DEPRECATION_WARNINGS
 FCoreDelegates::FOnMountPak FCoreDelegates::OnMountPak;
+PRAGMA_ENABLE_DEPRECATION_WARNINGS
+FCoreDelegates::FMountPak FCoreDelegates::MountPak;
 FCoreDelegates::FOnUnmountPak FCoreDelegates::OnUnmountPak;
 FCoreDelegates::FOnOptimizeMemoryUsageForMountedPaks FCoreDelegates::OnOptimizeMemoryUsageForMountedPaks;
 
-FCoreDelegates::FOnPakFileMounted FCoreDelegates::OnPakFileMounted;
+FCoreDelegates::FOnPakFileMounted2 FCoreDelegates::OnPakFileMounted2;
 PRAGMA_DISABLE_DEPRECATION_WARNINGS
+FCoreDelegates::FOnPakFileMounted FCoreDelegates::OnPakFileMounted;
 FCoreDelegates::FPakFileMountedDelegate FCoreDelegates::PakFileMountedCallback;
 PRAGMA_ENABLE_DEPRECATION_WARNINGS
 FCoreDelegates::FNewFileAddedDelegate FCoreDelegates::NewFileAddedDelegate;
diff --git a/Engine/Source/Runtime/Core/Public/GenericPlatform/GenericPlatformFile.h b/Engine/Source/Runtime/Core/Public/GenericPlatform/GenericPlatformFile.h
--- a/Engine/Source/Runtime/Core/Public/GenericPlatform/GenericPlatformFile.h
+++ b/Engine/Source/Runtime/Core/Public/GenericPlatform/GenericPlatformFile.h
@@ -237,7 +237,7 @@ public:
 	/**
 	 * Build an in memory unique pak file from a subset of files in this pak file
 	 */
-	virtual void		MakeUniquePakFilesForTheseFiles(TArray<TArray<FString>> InFiles) { }
+	virtual void		MakeUniquePakFilesForTheseFiles(const TArray<TArray<FString>>& InFiles) { }
 
 	/**
 	* Performs initialization of the platform file after the new async IO has been enabled
@@ -599,3 +599,22 @@ public:
 	}
 	//~ End IPlatformFile Interface
 };
+
+/* Interface class for FPakFile to allow usage from modules that cannot have a compile dependency on FPakFile */
+class CORE_API IPakFile
+{
+public:
+	virtual const FString& PakGetPakFilename() const = 0;
+	/**
+	  * Return whether the Pak has an entry for the given FileName.  Not necessarily exclusive; other Patch Paks may have their own copy of the same File.
+	  * @param Filename The full LongPackageName path to the file, as returned from FPackageName::LongPackageNameToFilename + extension.  Comparison is case-insensitive.
+	  */
+	virtual bool PakContains(const FString& Filename) const = 0;
+	virtual int32 PakGetPakchunkIndex() const = 0;
+	/**
+	 * Calls the given Visitor on every FileName in the Pruned Directory Index. FileNames passed to the Vistory are the RelativePath from the Mount of the PakFile
+	 * The Pruned Directory Index at Runtime contains only the DirectoryIndexKeepFiles-specified subset of FilesNames and DirectoryNames that exist in the PakFile
+	 */
+	virtual void PakVisitPrunedFilenames(IPlatformFile::FDirectoryVisitor& Visitor) const = 0;
+	virtual const FString& PakGetMountPoint() const = 0;
+};
\ No newline at end of file
diff --git a/Engine/Source/Runtime/Core/Public/Internationalization/TextLocalizationManager.h b/Engine/Source/Runtime/Core/Public/Internationalization/TextLocalizationManager.h
--- a/Engine/Source/Runtime/Core/Public/Internationalization/TextLocalizationManager.h
+++ b/Engine/Source/Runtime/Core/Public/Internationalization/TextLocalizationManager.h
@@ -19,6 +19,7 @@
 
 struct FPolyglotTextData;
 class ILocalizedTextSource;
+class IPakFile;
 class FTextLocalizationResource;
 
 typedef TSharedRef<FString, ESPMode::ThreadSafe> FTextDisplayStringRef;
@@ -229,7 +230,7 @@ public:
 
 private:
 	/** Callback for when a PAK file is loaded. Loads any chunk specific localization resources. */
-	void OnPakFileMounted(const TCHAR* PakFilename, const int32 ChunkId);
+	void OnPakFileMounted(const IPakFile& PakFile);
 
 	/** Callback for changes in culture. Loads the new culture's localization resources. */
 	void OnCultureChanged();
diff --git a/Engine/Source/Runtime/Core/Public/Misc/CoreDelegates.h b/Engine/Source/Runtime/Core/Public/Misc/CoreDelegates.h
--- a/Engine/Source/Runtime/Core/Public/Misc/CoreDelegates.h
+++ b/Engine/Source/Runtime/Core/Public/Misc/CoreDelegates.h
@@ -13,6 +13,7 @@
 
 class AActor;
 class Error;
+class IPakFile;
 
 // delegates for hotfixes
 namespace EHotfixDelegates
@@ -60,21 +61,27 @@ public:
 	// delegate type for prompting the pak system to mount all pak files, which haven't already been mounted, from all default locations
 	DECLARE_DELEGATE_RetVal_OneParam(int32, FOnMountAllPakFiles, const TArray<FString>&);
 
-	// delegate type for prompting the pak system to mount a new pak
+	// deprecated delegate type for prompting the pak system to mount a new pak
 	DECLARE_DELEGATE_RetVal_ThreeParams(bool, FOnMountPak, const FString&, int32, IPlatformFile::FDirectoryVisitor*);
 
 	// delegate type for prompting the pak system to mount a new pak
+	DECLARE_DELEGATE_RetVal_TwoParams(IPakFile*, FMountPak, const FString&, int32);
+
+	// delegate type for prompting the pak system to unmount a pak
 	DECLARE_DELEGATE_RetVal_OneParam(bool, FOnUnmountPak, const FString&);
 
 	// delegate type for prompting the pak system to optimize memory for mounted paks
 	DECLARE_DELEGATE(FOnOptimizeMemoryUsageForMountedPaks);
 
-	// delegate for handling when a new pak file is successfully mounted passes in the name of the mounted pak file
+	// deprecated delegate for handling when a new pak file is successfully mounted passes in the name of the mounted pak file
 	DECLARE_MULTICAST_DELEGATE_OneParam(FPakFileMountedDelegate, const TCHAR*);
 
-	// delegate for handling when a new pak file is successfully mounted passes in the name of the pak file and its chunk ID (or INDEX_NONE)
+	// deprecated delegate for handling when a new pak file is successfully mounted passes in the name of the pak file and its chunk ID (or INDEX_NONE)
 	DECLARE_MULTICAST_DELEGATE_TwoParams(FOnPakFileMounted, const TCHAR*, const int32);
 
+	// delegate for handling when a new pak file is successfully mounted
+	DECLARE_MULTICAST_DELEGATE_OneParam(FOnPakFileMounted2, const IPakFile&);
+
 	// delegate to let other systems no that no paks were mounted, in case something wants to handle that case
 	DECLARE_MULTICAST_DELEGATE(FNoPakFilesMountedDelegate);
 
@@ -154,18 +161,25 @@ public:
 	// Callback for mounting all the pak files in default locations
 	static FOnMountAllPakFiles OnMountAllPakFiles;
 
-	// Callback for mounting a new pak file.
+	// Callback to prompt the pak system to mount a pak file
+	static FMountPak MountPak;
+
+	UE_DEPRECATED(4.26, "OnMountPak is deprecated; use MountPak instead.")
 	static FOnMountPak OnMountPak;
 
-	// Callback for unmounting a pak file.
+	// Callback to prompt the pak system to unmount a pak file.
 	static FOnUnmountPak OnUnmountPak;
 
 	// Callback to optimize memeory for currently mounted paks
 	static FOnOptimizeMemoryUsageForMountedPaks OnOptimizeMemoryUsageForMountedPaks;
 
 	// After a pakfile is mounted this is called
+	static FOnPakFileMounted2 OnPakFileMounted2;
+
+	UE_DEPRECATED(4.26, "FCoreDelegates::OnPakFileMounted is deprecated; use OnPakFileMounted2 instead")
 	static FOnPakFileMounted OnPakFileMounted;
-	UE_DEPRECATED(4.25, "FCoreDelegates::PakFileMountedCallback is deprecated. Use FCoreDelegates::OnPakFileMounted instead.")
+
+	UE_DEPRECATED(4.25, "FCoreDelegates::PakFileMountedCallback is deprecated. Use FCoreDelegates::OnPakFileMounted2 instead.")
 	static FPakFileMountedDelegate PakFileMountedCallback;
 
 	// After a file is added this is called
