Fix bad merge - restore diffs from revision 50. 51 (Pak file freezing) was intentionally reverted, but 50 (pak file system handle mounting of IoStores) was a mistake and should have been kept.
#rnx

diff --git a/Engine/Source/Developer/PakFileUtilities/Private/PakFileUtilities.cpp b/Engine/Source/Developer/PakFileUtilities/Private/PakFileUtilities.cpp
--- a/Engine/Source/Developer/PakFileUtilities/Private/PakFileUtilities.cpp
+++ b/Engine/Source/Developer/PakFileUtilities/Private/PakFileUtilities.cpp
@@ -5124,12 +5124,6 @@ bool ExecuteUnrealPak(const TCHAR* CmdLine)
 			return false;
 		}
 
-		if (Entries.Num() == 0)
-		{
-			UE_LOG(LogPakFile, Error, TEXT("No files specified to add to pak file."));
-			return false;
-		}
-
 		int32 LowestSourcePakVersion = 0;
 		TMap<FString, FFileInfo> SourceFileHashes;
 
