- Added -Info option for UnrealPak to dump some header info about a pak file

diff --git a/Engine/Source/Developer/PakFileUtilities/Private/PakFileUtilities.cpp b/Engine/Source/Developer/PakFileUtilities/Private/PakFileUtilities.cpp
--- a/Engine/Source/Developer/PakFileUtilities/Private/PakFileUtilities.cpp
+++ b/Engine/Source/Developer/PakFileUtilities/Private/PakFileUtilities.cpp
@@ -3569,6 +3569,29 @@ void CreateDiffRelativePathMap(TArray<FString>& FileNames, const FString& RootPa
 	}
 }
 
+bool DumpPakInfo(const FString& InPakFilename, const FKeyChain& InKeyChain)
+{
+	FPakFile PakFile(&FPlatformFileManager::Get().GetPlatformFile(), *InPakFilename, false);
+
+	if (!PakFile.IsValid())
+	{
+		return false;
+	}
+
+
+	const FPakInfo& Info = PakFile.GetInfo();
+
+	UE_LOG(LogPakFile, Display, TEXT("Pak File: %s"), *InPakFilename);
+	UE_LOG(LogPakFile, Display, TEXT("    Version: %d"), Info.Version);
+	UE_LOG(LogPakFile, Display, TEXT("    IndexOffset: %d"), Info.IndexOffset);
+	UE_LOG(LogPakFile, Display, TEXT("    IndexSize: %d"), Info.IndexSize);
+	UE_LOG(LogPakFile, Display, TEXT("    IndexHash: %s"), *Info.IndexHash.ToString());
+	UE_LOG(LogPakFile, Display, TEXT("    bEncryptedIndex: %d"), Info.bEncryptedIndex);
+	UE_LOG(LogPakFile, Display, TEXT("    EncryptionKeyGuid: %s"), *Info.EncryptionKeyGuid.ToString());
+
+	return true;
+}
+
 bool DiffFilesInPaks(const FString& InPakFilename1, const FString& InPakFilename2, const bool bLogUniques1, const bool bLogUniques2, const FKeyChain& InKeyChain)
 {
 	int32 NumUniquePAK1 = 0;
@@ -4752,6 +4775,19 @@ bool ExecuteUnrealPak(const TCHAR* CmdLine)
 		return ListFilesInPak(*PakFilename, SizeFilter, !bExcludeDeleted, *CSVFilename, bExtractToMountPoint, KeyChain);
 	}
 
+	if (FParse::Param(CmdLine, TEXT("Info")))
+	{
+		if (NonOptionArguments.Num() != 1)
+		{
+			UE_LOG(LogPakFile, Error, TEXT("Incorrect arguments. Expected: -Info <PakFile>"));
+			return false;
+		}
+
+		FString PakFilename = GetPakPath(*NonOptionArguments[0], false);
+
+		return DumpPakInfo(PakFilename, KeyChain);
+	}
+
 	if (FParse::Param(CmdLine, TEXT("Diff")))
 	{
 		if(NonOptionArguments.Num() != 2)
@@ -5147,6 +5183,7 @@ bool ExecuteUnrealPak(const TCHAR* CmdLine)
 	UE_LOG(LogPakFile, Error, TEXT("No pak file name specified. Usage:"));
 	UE_LOG(LogPakFile, Error, TEXT("  UnrealPak <PakFilename> -Test"));
 	UE_LOG(LogPakFile, Error, TEXT("  UnrealPak <PakFilename> -Verify"));
+	UE_LOG(LogPakFile, Error, TEXT("  UnrealPak <PakFilename> -Info"));
 	UE_LOG(LogPakFile, Error, TEXT("  UnrealPak <PakFilename> -List [-ExcludeDeleted]"));
 	UE_LOG(LogPakFile, Error, TEXT("  UnrealPak <PakFilename> <GameUProjectName> <GameFolderName> -ExportDependencies=<OutputFileBase> -NoAssetRegistryCache -ForceDependsGathering"));
 	UE_LOG(LogPakFile, Error, TEXT("  UnrealPak <PakFilename> -Extract <ExtractDir> [-Filter=<filename>]"));
