Extracting pak files now supports multiple encryption keys
Graeme.Thornton

diff --git a/Engine/Source/Developer/PakFileUtilities/Private/PakFileUtilities.cpp b/Engine/Source/Developer/PakFileUtilities/Private/PakFileUtilities.cpp
--- a/Engine/Source/Developer/PakFileUtilities/Private/PakFileUtilities.cpp
+++ b/Engine/Source/Developer/PakFileUtilities/Private/PakFileUtilities.cpp
@@ -1443,7 +1443,11 @@ bool UncompressCopyFile(FArchive& Dest, FArchive& Source, const FPakEntry& Entry
 
 		if (Entry.IsEncrypted())
 		{
-			const FNamedAESKey* Key = InKeyChain.MasterEncryptionKey;
+			FNamedAESKey* Key = InKeyChain.EncryptionKeys.Find(PakFile.GetInfo().EncryptionKeyGuid);
+			if (Key == nullptr)
+			{
+				Key = InKeyChain.MasterEncryptionKey;
+			}
 			check(Key);
 			FAES::DecryptData(PersistentBuffer, SizeToRead, Key->Key);
 		}
