Fixed compile error in my last checkin.

diff --git a/Engine/Source/Developer/PakFileUtilities/Private/PakFileUtilities.cpp b/Engine/Source/Developer/PakFileUtilities/Private/PakFileUtilities.cpp
--- a/Engine/Source/Developer/PakFileUtilities/Private/PakFileUtilities.cpp
+++ b/Engine/Source/Developer/PakFileUtilities/Private/PakFileUtilities.cpp
@@ -1443,7 +1443,7 @@ bool UncompressCopyFile(FArchive& Dest, FArchive& Source, const FPakEntry& Entry
 
 		if (Entry.IsEncrypted())
 		{
-			FNamedAESKey* Key = InKeyChain.EncryptionKeys.Find(PakFile.GetInfo().EncryptionKeyGuid);
+			const FNamedAESKey* Key = InKeyChain.EncryptionKeys.Find(PakFile.GetInfo().EncryptionKeyGuid);
 			if (Key == nullptr)
 			{
 				Key = InKeyChain.MasterEncryptionKey;
