- Added compression methods to the pak dump outpout

diff --git a/Engine/Source/Developer/PakFileUtilities/Private/PakFileUtilities.cpp b/Engine/Source/Developer/PakFileUtilities/Private/PakFileUtilities.cpp
--- a/Engine/Source/Developer/PakFileUtilities/Private/PakFileUtilities.cpp
+++ b/Engine/Source/Developer/PakFileUtilities/Private/PakFileUtilities.cpp
@@ -3592,6 +3592,11 @@ bool DumpPakInfo(const FString& InPakFilename, const FKeyChain& InKeyChain)
 	UE_LOG(LogPakFile, Display, TEXT("    IndexHash: %s"), *Info.IndexHash.ToString());
 	UE_LOG(LogPakFile, Display, TEXT("    bEncryptedIndex: %d"), Info.bEncryptedIndex);
 	UE_LOG(LogPakFile, Display, TEXT("    EncryptionKeyGuid: %s"), *Info.EncryptionKeyGuid.ToString());
+	UE_LOG(LogPakFile, Display, TEXT("    CompressionMethods:"));
+	for (FName Method : Info.CompressionMethods)
+	{
+		UE_LOG(LogPakFile, Display, TEXT("        %s"), *Method.ToString());
+	}
 
 	return true;
 }
