improvements to Pak file for modern compressors

diff --git a/Engine/Source/Developer/PakFileUtilities/Private/PakFileUtilities.cpp b/Engine/Source/Developer/PakFileUtilities/Private/PakFileUtilities.cpp
--- a/Engine/Source/Developer/PakFileUtilities/Private/PakFileUtilities.cpp
+++ b/Engine/Source/Developer/PakFileUtilities/Private/PakFileUtilities.cpp
@@ -1833,8 +1833,6 @@ bool CreatePakFile(const TCHAR* Filename, TArray<FPakInputPair>& FilesToAdd, con
 
 	}
 
-	TArray<FName> UsedCompressionFormats; // List of compression formats we actually used in this pak file (used for logging only)
-
 	if (InKeyChain.MasterEncryptionKey)
 	{
 		UE_LOG(LogPakFile, Display, TEXT("Using encryption key '%s' [%s]"), *InKeyChain.MasterEncryptionKey->Name, *InKeyChain.MasterEncryptionKey->Guid.ToString());
@@ -1869,6 +1867,30 @@ bool CreatePakFile(const TCHAR* Filename, TArray<FPakInputPair>& FilesToAdd, con
 	uint64 TotalRequestedEncryptedFiles = 0;
 	uint64 TotalEncryptedFiles = 0;
 	uint64 TotalEncryptedDataSize = 0;
+	
+	// track compression stats per format :
+	TArray<FName> CompressionFormatsAndNone = CmdLineParameters.CompressionFormats;
+	CompressionFormatsAndNone.AddUnique(NAME_None);
+	int32 CompressionFormatsAndNone_Num = CompressionFormatsAndNone.Num();
+	TArray<int32> Compressor_Stat_Count; Compressor_Stat_Count.SetNumZeroed(CompressionFormatsAndNone_Num);
+	TArray<uint64> Compressor_Stat_RawBytes; Compressor_Stat_RawBytes.SetNumZeroed(CompressionFormatsAndNone_Num);
+	TArray<uint64> Compressor_Stat_CompBytes; Compressor_Stat_CompBytes.SetNumZeroed(CompressionFormatsAndNone_Num);
+
+	{
+		// log the methods and indexes
+		// we're going to prefer to use only index [0] so we want that to be our favorite compressor
+		FString log_line(TEXT("CompressionFormats in priority order: "));
+		for (int32 MethodIndex = 0; MethodIndex < CmdLineParameters.CompressionFormats.Num(); MethodIndex++)
+		{
+			FName CompressionMethod = CompressionFormatsAndNone[MethodIndex];
+			if ( MethodIndex > 0 )
+			{
+				log_line += TEXT(", ");
+			}
+			log_line += CompressionMethod.ToString();
+		}
+		UE_LOG(LogPakFile, Display, TEXT("%s"), *log_line);
+	}
 
 	TArray<FString> ExtensionsToNotUsePluginCompression;
 	GConfig->GetArray(TEXT("Pak"), TEXT("ExtensionsToNotUsePluginCompression"), ExtensionsToNotUsePluginCompression, GEngineIni);
@@ -1946,7 +1968,9 @@ bool CreatePakFile(const TCHAR* Filename, TArray<FPakInputPair>& FilesToAdd, con
 			OriginalFileSize = IFileManager::Get().FileSize(*FileToAdd->Source);
 			RealFileSize = OriginalFileSize + Entry.Info.GetSerializedSize(FPakInfo::PakFile_Version_Latest);
 
-			if (OriginalFileSize <= 0 || (FileToAdd->bNeedsCompression == false))
+			// don't try to compress tiny files
+			// even if they do compress, it is a bad use of decoder time
+			if (OriginalFileSize < 1024 || (FileToAdd->bNeedsCompression == false))
 			{
 				// done, don't need to do anything else
 				Complete();
@@ -1954,7 +1978,12 @@ bool CreatePakFile(const TCHAR* Filename, TArray<FPakInputPair>& FilesToAdd, con
 			}
 
 			bool bSomeCompressionSucceeded = false;
-			for (int32 MethodIndex = 0; MethodIndex < CompressionFormats->Num(); MethodIndex++)
+			
+			// if first method (oodle) doesn't compress, do NOT try other methods (zlib)
+			// if Oodle refused to compress it was not an error, it was a choice because we didn't like compressing that file
+			//	we do NOT want Zlib to then get enabled for that file!
+			for (int32 MethodIndex = 0; MethodIndex < 1; MethodIndex++)
+			//for (int32 MethodIndex = 0; MethodIndex < CompressionFormats->Num(); MethodIndex++)
 			{
 				CompressionMethod = (*CompressionFormats)[MethodIndex];
 
@@ -1971,11 +2000,27 @@ bool CreatePakFile(const TCHAR* Filename, TArray<FPakInputPair>& FilesToAdd, con
 				// attempt to compress the data
 				if (CompressedFileBuffer.CompressFileToWorkingBuffer(*FileToAdd, CompressionMethod, CompressionBlockSize))
 				{
-					// Check the compression ratio, if it's too low just store uncompressed. Also take into account read size
-					// if we still save 64KB it's probably worthwhile compressing, as that saves a file read operation in the runtime.
-								// TODO: drive this threshold from the command line
-					float PercentLess = ((float)CompressedFileBuffer.TotalCompressedSize / (OriginalFileSize / 100.f));
-					const bool bNotEnoughCompression = (PercentLess > 90.f) && ((OriginalFileSize - CompressedFileBuffer.TotalCompressedSize) < 65536);
+					// for modern compressors we don't want any funny heuristics turning compression on/off
+					// let the compressor decide; it has an introspective measure of whether compression is worth doing or not
+					bool bNotEnoughCompression = CompressedFileBuffer.TotalCompressedSize >= OriginalFileSize;
+
+					if ( CompressionMethod == NAME_Zlib )
+					{
+						// for forced-Zlib files still use the old heuristic :
+						// TODO : move this inside the zlib compressor
+
+						// Zlib must save at least 1K regardless of percentage (for small files)
+						bNotEnoughCompression = (OriginalFileSize - CompressedFileBuffer.TotalCompressedSize) < 1024;
+						if ( ! bNotEnoughCompression )
+						{
+							// Check the compression ratio, if it's too low just store uncompressed. Also take into account read size
+							// if we still save 64KB it's probably worthwhile compressing, as that saves a file read operation in the runtime.
+							// TODO: drive this threshold from the command line
+							float PercentLess = ((float)CompressedFileBuffer.TotalCompressedSize / (OriginalFileSize / 100.f));
+							bNotEnoughCompression = (PercentLess > 90.f) && ((OriginalFileSize - CompressedFileBuffer.TotalCompressedSize) < 65536);
+						}
+					}
+
 					const bool bIsLastCompressionFormat = MethodIndex == CompressionFormats->Num() - 1;
 					if (bNotEnoughCompression && (!bForceCompress || !bIsLastCompressionFormat))
 					{
@@ -2052,29 +2097,46 @@ bool CreatePakFile(const TCHAR* Filename, TArray<FPakInputPair>& FilesToAdd, con
 	const bool bRunAsync = CmdLineParameters.bAsyncCompression;
 
 	GetDerivedDataCacheRef();
+	
+	TArray<FFileRegion> AllFileRegions;
+	
+	// limit memory use :
+	// Launch async tasks ahead of our current retiring index
+	// enough to keep the cores busy
+	// but we don't want every file in the whole Pak
+	// because the compressed buffers are held after the task is done
+
+	int NextLaunchIndex = 0;
+	int NumToLaunchAhead = 128;   // <- make NumToLaunchAhead lower to reduce memory use more
+	// TODO: NumToLaunchAhead should ideally scale up with 2*cores or so but also be limited by memory size.
 
 	for (int32 FileIndex = 0; FileIndex < FilesToAdd.Num(); FileIndex++)
 	{
-		AsyncCompressors[FileIndex].Init(&FilesToAdd[FileIndex], CmdLineParameters, &NoPluginCompressionExtensions, &NoPluginCompressionFileNames);
-		if (bRunAsync)
+		// FileIndex is what we're currently retiring
+		// everything in [FileIndex,NextLaunchIndex) may be pending
+		// launch new jobs up to NumToLaunchAhead ahead of FileIndex
+		int32 LaunchFileIndexEnd = FMath::Min( FileIndex + NumToLaunchAhead, FilesToAdd.Num() );
+		for (int32 LaunchFileIndex = NextLaunchIndex; LaunchFileIndex < LaunchFileIndexEnd; LaunchFileIndex++)
 		{
-			if (FilesToAdd[FileIndex].bNeedsCompression)
-			{
-				(new FAutoDeleteAsyncTask<FRunCompressionTask>(&AsyncCompressors[FileIndex]))->StartBackgroundTask();
-			}
-			else
+			AsyncCompressors[LaunchFileIndex].Init(&FilesToAdd[LaunchFileIndex], CmdLineParameters, &NoPluginCompressionExtensions, &NoPluginCompressionFileNames);
+			
+			if (bRunAsync)
 			{
-				// call compress function inline 
-				// it won't do anything except for initialize some internal variables used in the non compressed path
-				// we don't want to pass these to a different thread as they may cause congestion with legitimate tasks
-				AsyncCompressors[FileIndex].Compress();
+				if (FilesToAdd[LaunchFileIndex].bNeedsCompression)
+				{
+					(new FAutoDeleteAsyncTask<FRunCompressionTask>(&AsyncCompressors[LaunchFileIndex]))->StartBackgroundTask();
+				}
+				else
+				{
+					// call compress function inline 
+					// it won't do anything except for initialize some internal variables used in the non compressed path
+					// we don't want to pass these to a different thread as they may cause congestion with legitimate tasks
+					AsyncCompressors[LaunchFileIndex].Compress();
+				}
 			}
 		}
-	}
+		NextLaunchIndex = LaunchFileIndexEnd;
 
-	TArray<FFileRegion> AllFileRegions;
-	for (int32 FileIndex = 0; FileIndex < FilesToAdd.Num(); FileIndex++)
-	{
 		bool bDeleted = FilesToAdd[FileIndex].bIsDeleteRecord;
 		bool bIsUAssetUExpPairUAsset = false;
 		bool bIsUAssetUExpPairUExp = false;
@@ -2270,9 +2332,22 @@ bool CreatePakFile(const TCHAR* Filename, TArray<FPakInputPair>& FilesToAdd, con
 			}
 			if (FilesToAdd[FileIndex].bNeedsCompression && CompressionMethod != NAME_None)
 			{
-				UsedCompressionFormats.AddUnique(CompressionMethod); // used for logging only
 				FinalizeCopyCompressedFileToPak(Info, CompressedFileBuffer, NewEntry);
 			}
+			
+			{
+				// track per-compressor stats :
+				// note GetCompressionMethodIndex in the pak entry is the index in the Pak file list of compressors
+				//	not the same as the index in the command line list
+
+				int32 CompressionMethodIndex;
+				if ( CompressionFormatsAndNone.Find(CompressionMethod,CompressionMethodIndex) )
+				{
+					Compressor_Stat_Count[CompressionMethodIndex] += 1;
+					Compressor_Stat_RawBytes[ CompressionMethodIndex] += NewEntry.Info.UncompressedSize;
+					Compressor_Stat_CompBytes[CompressionMethodIndex] += NewEntry.Info.Size;
+				}
+			}
 
 			// Write to file
 			int64 Offset = PakFileHandle->Tell();
@@ -2352,6 +2427,17 @@ bool CreatePakFile(const TCHAR* Filename, TArray<FPakInputPair>& FilesToAdd, con
 	ReadBuffer = NULL;
 
 	FThreadLocalScratchSpace::Get().CleanUp();
+	
+	// log per-compressor stats :
+	for (int32 MethodIndex = 0; MethodIndex < CompressionFormatsAndNone_Num; MethodIndex++)
+	{
+		FName CompressionMethod = CompressionFormatsAndNone[MethodIndex];
+		UE_LOG(LogPakFile, Display, TEXT("CompressionFormat %d [%s] : %d files, %lld -> %lld bytes"), MethodIndex, *(CompressionMethod.ToString()),
+			Compressor_Stat_Count[MethodIndex],
+			Compressor_Stat_RawBytes[MethodIndex],
+			Compressor_Stat_CompBytes[MethodIndex]
+			);
+	}
 
 	auto FinalizeIndexBlockSize = [&Info](TArray<uint8>& IndexData)
 	{
@@ -2532,14 +2618,6 @@ bool CreatePakFile(const TCHAR* Filename, TArray<FPakInputPair>& FilesToAdd, con
 		float PercentLess = ((float)TotalCompressedSize / (TotalUncompressedSize / 100.f));
 		UE_LOG(LogPakFile, Display, TEXT("Compression summary: %.2f%% of original size. Compressed Size %lld bytes, Original Size %lld bytes. "), PercentLess, TotalCompressedSize, TotalUncompressedSize);
 
-		FString UsedCompressionFormatsString;
-		for (FName CompressionFormat : UsedCompressionFormats)
-		{
-			UsedCompressionFormatsString.Append( CompressionFormat.ToString() + TEXT(", ") );
-		}
-
-		UE_LOG(LogPakFile, Display, TEXT("Used compression formats (in priority order) '%s'"), *UsedCompressionFormatsString);
-
 		if (GTotalFilesWithPoorForcedCompression > 0)
 		{
 			UE_LOG(LogPakFile, Display, TEXT("Num files forcibly compressed due to -forcecompress option: %i, using %i bytes extra"), (int64)GTotalFilesWithPoorForcedCompression, (int64)GTotalExtraMemoryForPoorForcedCompression);
