Fix UnrealPak compile from pak file shared reader changes.

diff --git a/Engine/Source/Developer/PakFileUtilities/Private/PakFileUtilities.cpp b/Engine/Source/Developer/PakFileUtilities/Private/PakFileUtilities.cpp
--- a/Engine/Source/Developer/PakFileUtilities/Private/PakFileUtilities.cpp
+++ b/Engine/Source/Developer/PakFileUtilities/Private/PakFileUtilities.cpp
@@ -3223,7 +3223,7 @@ bool ListFilesAtOffset( const TCHAR* InPakFileName, const TArray<int64>& InOffse
 	UE_LOG( LogPakFile, Display, TEXT("%-12s%-12s%-12s%s"), TEXT("Offset"), TEXT("File Offset"), TEXT("File Size"), TEXT("File Name") );
 
 	TArray<int64> OffsetsToCheck = InOffsets;
-	FArchive& PakReader = *PakFile.GetSharedReader(NULL);
+	FSharedPakReader PakReader = PakFile.GetSharedReader(NULL);
 	for (FPakFile::FPakEntryIterator It(PakFile); It; ++It)
 	{
 		const FPakEntry& Entry = It.Info();
@@ -3270,7 +3270,7 @@ bool ShowCompressionBlockCRCs( const TCHAR* InPakFileName, TArray<int64>& InOffs
 	}
 
 	// read the pak file and iterate over all given offsets
-	FArchive& PakReader = *PakFile.GetSharedReader(NULL);
+	FSharedPakReader PakReader = PakFile.GetSharedReader(NULL);
 	UE_LOG(LogPakFile, Display, TEXT("") );
 	for( int64 Offset : InOffsets )
 	{
@@ -3331,9 +3331,9 @@ bool ShowCompressionBlockCRCs( const TCHAR* InPakFileName, TArray<int64>& InOffs
 		{
 			uint32 CompressedBlockSize = Entry->CompressionBlocks[BlockIndex].CompressedEnd - Entry->CompressionBlocks[BlockIndex].CompressedStart;
 			uint32 UncompressedBlockSize = (uint32)FMath::Min<int64>(Entry->UncompressedSize - Entry->CompressionBlockSize*BlockIndex, Entry->CompressionBlockSize);
-			PakReader.Seek(Entry->CompressionBlocks[BlockIndex].CompressedStart + (PakFile.GetInfo().HasRelativeCompressedChunkOffsets() ? Entry->Offset : 0));
+			PakReader->Seek(Entry->CompressionBlocks[BlockIndex].CompressedStart + (PakFile.GetInfo().HasRelativeCompressedChunkOffsets() ? Entry->Offset : 0));
 			uint32 SizeToRead = Entry->IsEncrypted() ? Align(CompressedBlockSize, FAES::AESBlockSize) : CompressedBlockSize;
-			PakReader.Serialize(PersistentBuffer, SizeToRead);
+			PakReader->Serialize(PersistentBuffer, SizeToRead);
 
 			if (Entry->IsEncrypted())
 			{
@@ -3408,7 +3408,7 @@ bool GeneratePIXMappingFile(const TArray<FString> InPakFileList, const FString&
 		CSVFileWriter->Logf(TEXT("%s"), *PakFileName);
 
 		const FString PakFileMountPoint = PakFile.GetMountPoint();
-		FArchive& PakReader = *PakFile.GetSharedReader(NULL);
+		FSharedPakReader PakReader = PakFile.GetSharedReader(NULL);
 		if (!PakFile.HasFilenames())
 		{
 			UE_LOG(LogPakFile, Error, TEXT("PakFiles were loaded without filenames, cannot generate PIX mapping file."));
@@ -3478,7 +3478,7 @@ bool ExtractFilesFromPak(const TCHAR* InPakFilename, TMap<FString, FFileInfo>& I
 		if (PakFile.IsValid())
 		{
 			FString DestPath(InDestPath);
-			FArchive& PakReader = *PakFile.GetSharedReader(NULL);
+			FSharedPakReader PakReader = PakFile.GetSharedReader(NULL);
 			const int64 BufferSize = 8 * 1024 * 1024; // 8MB buffer for extracting
 			void* Buffer = FMemory::Malloc(BufferSize);
 			int64 CompressionBufferSize = 0;
@@ -3535,10 +3535,10 @@ bool ExtractFilesFromPak(const TCHAR* InPakFilename, TMap<FString, FFileInfo>& I
 						continue;
 					}
 
-					PakReader.Seek(Entry.Offset);
+					PakReader->Seek(Entry.Offset);
 					uint32 SerializedCrcTest = 0;
 					FPakEntry EntryInfo;
-					EntryInfo.Serialize(PakReader, PakFile.GetInfo().Version);
+					EntryInfo.Serialize(PakReader.GetArchive(), PakFile.GetInfo().Version);
 					if (EntryInfo.IndexDataEquals(Entry))
 					{
 						TUniquePtr<FArchive> FileHandle(IFileManager::Get().CreateFileWriter(*DestFilename));
@@ -3546,11 +3546,11 @@ bool ExtractFilesFromPak(const TCHAR* InPakFilename, TMap<FString, FFileInfo>& I
 						{
 							if (Entry.CompressionMethodIndex == 0)
 							{
-								BufferedCopyFile(*FileHandle, PakReader, PakFile, Entry, Buffer, BufferSize, InKeyChain);
+								BufferedCopyFile(*FileHandle, PakReader.GetArchive(), PakFile, Entry, Buffer, BufferSize, InKeyChain);
 							}
 							else
 							{
-								UncompressCopyFile(*FileHandle, PakReader, Entry, PersistantCompressionBuffer, CompressionBufferSize, InKeyChain, PakFile);
+								UncompressCopyFile(*FileHandle, PakReader.GetArchive(), Entry, PersistantCompressionBuffer, CompressionBufferSize, InKeyChain, PakFile);
 							}
 							UE_LOG(LogPakFile, Display, TEXT("Extracted \"%s\" to \"%s\" Offset %d."), *EntryFileName, *DestFilename, Entry.Offset);
 							ExtractedCount++;
@@ -3665,8 +3665,8 @@ bool DiffFilesInPaks(const FString& InPakFilename1, const FString& InPakFilename
 	FPakFile& PakFile2 = *PakFilePtr2;
 	if (PakFile1.IsValid() && PakFile2.IsValid())
 	{		
-		FArchive& PakReader1 = *PakFile1.GetSharedReader(NULL);
-		FArchive& PakReader2 = *PakFile2.GetSharedReader(NULL);
+		FSharedPakReader PakReader1 = PakFile1.GetSharedReader(NULL);
+		FSharedPakReader PakReader2 = PakFile2.GetSharedReader(NULL);
 
 		const int64 BufferSize = 8 * 1024 * 1024; // 8MB buffer for extracting
 		void* Buffer = FMemory::Malloc(BufferSize);
@@ -3687,10 +3687,10 @@ bool DiffFilesInPaks(const FString& InPakFilename1, const FString& InPakFilename
 
 			//double check entry info and move pakreader into place
 			const FPakEntry& Entry1 = It.Info();
-			PakReader1.Seek(Entry1.Offset);
+			PakReader1->Seek(Entry1.Offset);
 
 			FPakEntry EntryInfo1;
-			EntryInfo1.Serialize(PakReader1, PakFile1.GetInfo().Version);
+			EntryInfo1.Serialize(PakReader1.GetArchive(), PakFile1.GetInfo().Version);
 
 			if (EntryInfo1 != Entry1)
 			{
@@ -3712,9 +3712,9 @@ bool DiffFilesInPaks(const FString& InPakFilename1, const FString& InPakFilename
 			}
 
 			//double check entry info and move pakreader into place
-			PakReader2.Seek(Entry2.Offset);
+			PakReader2->Seek(Entry2.Offset);
 			FPakEntry EntryInfo2;
-			EntryInfo2.Serialize(PakReader2, PakFile2.GetInfo().Version);
+			EntryInfo2.Serialize(PakReader2.GetArchive(), PakFile2.GetInfo().Version);
 			if (EntryInfo2 != Entry2)
 			{
 				UE_LOG(LogPakFile, Log, TEXT("PakEntry2Invalid, %s, 0, 0"), *PAK1FileName);
@@ -3735,20 +3735,20 @@ bool DiffFilesInPaks(const FString& InPakFilename1, const FString& InPakFilename
 
 				if (EntryInfo1.CompressionMethodIndex == 0)
 				{
-					BufferedCopyFile(PAKWriter1, PakReader1, PakFile1, Entry1, Buffer, BufferSize, InKeyChain);
+					BufferedCopyFile(PAKWriter1, PakReader1.GetArchive(), PakFile1, Entry1, Buffer, BufferSize, InKeyChain);
 				}
 				else
 				{
-					UncompressCopyFile(PAKWriter1, PakReader1, Entry1, PersistantCompressionBuffer, CompressionBufferSize, InKeyChain, PakFile1);
+					UncompressCopyFile(PAKWriter1, PakReader1.GetArchive(), Entry1, PersistantCompressionBuffer, CompressionBufferSize, InKeyChain, PakFile1);
 				}
 
 				if (EntryInfo2.CompressionMethodIndex == 0)
 				{
-					BufferedCopyFile(PAKWriter2, PakReader2, PakFile2, Entry2, Buffer, BufferSize, InKeyChain);
+					BufferedCopyFile(PAKWriter2, PakReader2.GetArchive(), PakFile2, Entry2, Buffer, BufferSize, InKeyChain);
 				}
 				else
 				{
-					UncompressCopyFile(PAKWriter2, PakReader2, Entry2, PersistantCompressionBuffer, CompressionBufferSize, InKeyChain, PakFile2);
+					UncompressCopyFile(PAKWriter2, PakReader2.GetArchive(), Entry2, PersistantCompressionBuffer, CompressionBufferSize, InKeyChain, PakFile2);
 				}
 
 				if (FMemory::Memcmp(PAKWriter1.GetData(), PAKWriter2.GetData(), EntryInfo1.UncompressedSize) != 0)
@@ -3767,10 +3767,10 @@ bool DiffFilesInPaks(const FString& InPakFilename1, const FString& InPakFilename
 		for (FPakFile::FPakEntryIterator It(PakFile2); It; ++It, ++FileCount)
 		{
 			const FPakEntry& Entry2 = It.Info();
-			PakReader2.Seek(Entry2.Offset);
+			PakReader2->Seek(Entry2.Offset);
 
 			FPakEntry EntryInfo2;
-			EntryInfo2.Serialize(PakReader2, PakFile2.GetInfo().Version);
+			EntryInfo2.Serialize(PakReader2.GetArchive(), PakFile2.GetInfo().Version);
 
 			if (EntryInfo2.IndexDataEquals(Entry2))
 			{
@@ -3857,7 +3857,7 @@ bool GenerateHashesFromPak(const TCHAR* InPakFilename, const TCHAR* InDestPakFil
 			{
 				OutUsedEncryptionKeys->Add( PakFile.GetInfo().EncryptionKeyGuid);
 			}
-			FArchive& PakReader = *PakFile.GetSharedReader(NULL);
+			FSharedPakReader PakReader = PakFile.GetSharedReader(NULL);
 			const int64 BufferSize = 8 * 1024 * 1024; // 8MB buffer for extracting
 			void* Buffer = FMemory::Malloc(BufferSize);
 			int64 CompressionBufferSize = 0;
@@ -3901,10 +3901,10 @@ bool GenerateHashesFromPak(const TCHAR* InPakFilename, const TCHAR* InDestPakFil
 				}
 				else
 				{
-				    PakReader.Seek(Entry.Offset);
+				    PakReader->Seek(Entry.Offset);
 				    uint32 SerializedCrcTest = 0;
 				    FPakEntry EntryInfo;
-				    EntryInfo.Serialize(PakReader, PakFile.GetInfo().Version);
+				    EntryInfo.Serialize(PakReader.GetArchive(), PakFile.GetInfo().Version);
 				    if (EntryInfo.IndexDataEquals(Entry))
 				    {
 					    // TUniquePtr<FArchive> FileHandle(IFileManager::Get().CreateFileWriter(*DestFilename));
@@ -3915,11 +3915,11 @@ bool GenerateHashesFromPak(const TCHAR* InPakFilename, const TCHAR* InDestPakFil
 					    {
 							if (Entry.CompressionMethodIndex == 0)
 						    {
-							    BufferedCopyFile(*FileHandle, PakReader, PakFile, Entry, Buffer, BufferSize, InKeyChain);
+							    BufferedCopyFile(*FileHandle, PakReader.GetArchive(), PakFile, Entry, Buffer, BufferSize, InKeyChain);
 						    }
 						    else
 						    {
-							    UncompressCopyFile(*FileHandle, PakReader, Entry, PersistantCompressionBuffer, CompressionBufferSize, InKeyChain, PakFile);
+							    UncompressCopyFile(*FileHandle, PakReader.GetArchive(), Entry, PersistantCompressionBuffer, CompressionBufferSize, InKeyChain, PakFile);
 						    }
     
 						    UE_LOG(LogPakFile, Display, TEXT("Generated hash for \"%s\""), *FullFilename);
