Implement copy/assignment operators for FKeyChain to prevent copying/using Key pointer from other instance

#6925

diff --git a/Engine/Source/Developer/PakFileUtilities/Private/Tests/KeyChainTest.cpp b/Engine/Source/Developer/PakFileUtilities/Private/Tests/KeyChainTest.cpp
new file mode 100644
--- /dev/null
+++ b/Engine/Source/Developer/PakFileUtilities/Private/Tests/KeyChainTest.cpp
@@ -0,0 +1,152 @@
+// Copyright Epic Games, Inc. All Rights Reserved.
+
+#include "Misc/KeyChainUtilities.h"
+#include "Misc/AutomationTest.h"
+
+#if WITH_DEV_AUTOMATION_TESTS 
+
+IMPLEMENT_SIMPLE_AUTOMATION_TEST(FKeyChainTest, "System.Core.Misc.KeyChain", EAutomationTestFlags::ApplicationContextMask | EAutomationTestFlags::EngineFilter)
+
+/**
+ * KeyChain tests - implemented in this module because of incorrect dependencies in KeyChainUtilities.h.
+ */
+bool FKeyChainTest::RunTest(const FString& Parameters)
+{
+	const FGuid DefaultGuid;
+	uint32 SigningKeyDummy = 42;
+	const FRSAKeyHandle DefaultSigningKey = reinterpret_cast<const FRSAKeyHandle>(&SigningKeyDummy);
+
+	// Default construct
+	{
+		FKeyChain Chain;
+		TestTrue(TEXT("Default construct - master key is invalid"), Chain.MasterEncryptionKey == nullptr); 
+		TestTrue(TEXT("Default construct - signing key is invalid"), Chain.SigningKey == nullptr); 
+		TestTrue(TEXT("Default construct - encryption keys are empty"), Chain.EncryptionKeys.Num() == 0); 
+	}
+
+	// Copy construct
+	{
+		FKeyChain Chain;
+		Chain.EncryptionKeys.Add(DefaultGuid, FNamedAESKey { TEXT("Default"), DefaultGuid, FAES::FAESKey() });
+		Chain.MasterEncryptionKey = Chain.EncryptionKeys.Find(DefaultGuid);
+		Chain.SigningKey = DefaultSigningKey;
+
+		FKeyChain Copy(Chain);
+
+		TestTrue(TEXT("Copy construct - with valid master key does NOT copy pointer"),
+			Copy.MasterEncryptionKey != Chain.MasterEncryptionKey); 
+		TestTrue(TEXT("Copy construct - master key name matches"),
+			Copy.MasterEncryptionKey->Name == Chain.MasterEncryptionKey->Name); 
+		TestTrue(TEXT("Copy construct - master key GUID name matches"),
+			Copy.MasterEncryptionKey->Guid == Chain.MasterEncryptionKey->Guid); 
+		TestTrue(TEXT("Copy construct - with valid master key sets master key"),
+			Copy.MasterEncryptionKey == Copy.EncryptionKeys.Find(DefaultGuid)); 
+		TestTrue(TEXT("Copy construct - signing key matches"),
+			Copy.SigningKey == Chain.SigningKey ); 
+	}
+
+	// Copy assign
+	{
+		{
+			FKeyChain Chain;
+			Chain.EncryptionKeys.Add(DefaultGuid, FNamedAESKey { TEXT("Default"), DefaultGuid, FAES::FAESKey() });
+			Chain.MasterEncryptionKey = Chain.EncryptionKeys.Find(DefaultGuid);
+			Chain.SigningKey = DefaultSigningKey;
+
+			FKeyChain Copy;
+			Copy = Chain;
+			
+			TestTrue(TEXT("Copy assign - with valid master key does NOT copy pointer"),
+				Copy.MasterEncryptionKey != Chain.MasterEncryptionKey); 
+			TestTrue(TEXT("Copy assign - master key name matches"),
+				Copy.MasterEncryptionKey->Name == Chain.MasterEncryptionKey->Name); 
+			TestTrue(TEXT("Copy assign - master key GUID name matches"),
+				Copy.MasterEncryptionKey->Guid == Chain.MasterEncryptionKey->Guid); 
+			TestTrue(TEXT("Copy assign - with valid master key sets master key"),
+				Copy.MasterEncryptionKey == Copy.EncryptionKeys.Find(DefaultGuid)); 
+			TestTrue(TEXT("Copy assign - signing key matches"),
+				Copy.SigningKey == Chain.SigningKey ); 
+		}
+
+		{
+			FKeyChain Copy;
+			Copy.EncryptionKeys.Add(DefaultGuid, FNamedAESKey { TEXT("Default"), DefaultGuid, FAES::FAESKey() });
+			Copy.MasterEncryptionKey = Copy.EncryptionKeys.Find(DefaultGuid);
+			Copy.SigningKey = DefaultSigningKey;
+
+			FKeyChain Chain;
+			Copy = Chain;
+			
+			TestTrue(TEXT("Copy assign - empty instance, clears master key"),
+				Copy.MasterEncryptionKey == nullptr); 
+			TestTrue(TEXT("Copy assign - empty instance, clears encryption keys"),
+				Copy.EncryptionKeys.Num() == 0); 
+			TestTrue(TEXT("Copy assign - signing key is invalid"),
+				Copy.SigningKey == nullptr); 
+		}
+	}
+
+	// Move construct
+	{
+		FKeyChain Moved;
+		Moved.EncryptionKeys.Add(DefaultGuid, FNamedAESKey { TEXT("Default"), DefaultGuid, FAES::FAESKey() });
+		Moved.MasterEncryptionKey = Moved.EncryptionKeys.Find(DefaultGuid);
+		Moved.SigningKey = DefaultSigningKey;
+
+		FKeyChain Chain(MoveTemp(Moved));
+
+		TestTrue(TEXT("Move construct - with valid master key sets master key"),
+			Chain.MasterEncryptionKey == Chain.EncryptionKeys.Find(DefaultGuid)); 
+		TestTrue(TEXT("Move construct - with valid master key sets signing key"),
+			Chain.SigningKey == DefaultSigningKey); 
+		TestTrue(TEXT("Move construct - invalidates moved instance"),
+			Moved.MasterEncryptionKey == nullptr); 
+		TestTrue(TEXT("Move construct - invalidates moved instance"),
+			Moved.SigningKey == InvalidRSAKeyHandle); 
+		TestTrue(TEXT("Move construct - invalidates moved instance"),
+			Moved.EncryptionKeys.Num() == 0); 
+	}
+
+	// Move assign
+	{
+		{
+			FKeyChain Moved;
+			Moved.EncryptionKeys.Add(DefaultGuid, FNamedAESKey { TEXT("Default"), DefaultGuid, FAES::FAESKey() });
+			Moved.MasterEncryptionKey = Moved.EncryptionKeys.Find(DefaultGuid);
+			Moved.SigningKey = DefaultSigningKey;
+
+			FKeyChain Chain;
+			Chain = MoveTemp(Moved);
+
+			TestTrue(TEXT("Move construct - with valid master key sets master key"),
+				Chain.MasterEncryptionKey == Chain.EncryptionKeys.Find(DefaultGuid)); 
+			TestTrue(TEXT("Move construct - with valid master key sets signing key"),
+				Chain.SigningKey == DefaultSigningKey); 
+			TestTrue(TEXT("Move construct - invalidates moved instance"),
+				Moved.MasterEncryptionKey == nullptr); 
+			TestTrue(TEXT("Move construct - invalidates moved instance"),
+				Moved.SigningKey == InvalidRSAKeyHandle); 
+			TestTrue(TEXT("Move construct - invalidates moved instance"),
+				Moved.EncryptionKeys.Num() == 0); 
+		}
+		
+		{
+			FKeyChain Copy;
+			Copy.EncryptionKeys.Add(DefaultGuid, FNamedAESKey { TEXT("Default"), DefaultGuid, FAES::FAESKey() });
+			Copy.MasterEncryptionKey = Copy.EncryptionKeys.Find(DefaultGuid);
+
+			Copy = FKeyChain();
+
+			TestTrue(TEXT("Move assign - empty instance, clears master key"),
+				Copy.MasterEncryptionKey == nullptr); 
+			TestTrue(TEXT("Move assign - empty instance, clears encryption keys"),
+				Copy.EncryptionKeys.Num() == 0);
+			TestTrue(TEXT("Move assign - invalidates signing key"),
+				Copy.SigningKey == InvalidRSAKeyHandle); 
+		}
+	}
+
+	return true;
+}
+
+#endif //WITH_DEV_AUTOMATION_TESTS
diff --git a/Engine/Source/Runtime/Core/Public/Misc/KeyChainUtilities.h b/Engine/Source/Runtime/Core/Public/Misc/KeyChainUtilities.h
--- a/Engine/Source/Runtime/Core/Public/Misc/KeyChainUtilities.h
+++ b/Engine/Source/Runtime/Core/Public/Misc/KeyChainUtilities.h
@@ -3,6 +3,7 @@
 #pragma once
 
 #include "CoreMinimal.h"
+#include "HAL/FileManager.h"
 #include "Misc/AES.h"
 #include "Misc/IEngineCrypto.h"
 #include "Misc/CoreDelegates.h"
@@ -25,6 +26,68 @@ struct FNamedAESKey
 
 struct FKeyChain
 {
+	FKeyChain() = default;
+
+	FKeyChain(const FKeyChain& Other)
+		: SigningKey(Other.SigningKey)
+		, EncryptionKeys(Other.EncryptionKeys)
+	{
+		if (Other.MasterEncryptionKey)
+		{
+			MasterEncryptionKey = EncryptionKeys.Find(Other.MasterEncryptionKey->Guid);
+		}
+	}
+	
+	FKeyChain(FKeyChain&& Other)
+		: SigningKey(Other.SigningKey)
+		, EncryptionKeys(MoveTemp(Other.EncryptionKeys))
+	{
+		if (Other.MasterEncryptionKey)
+		{
+			MasterEncryptionKey = EncryptionKeys.Find(Other.MasterEncryptionKey->Guid);
+		}
+		
+		Other.SigningKey = InvalidRSAKeyHandle;
+		Other.MasterEncryptionKey = nullptr;
+	}
+
+	FKeyChain& operator=(const FKeyChain& Other)
+	{
+		SigningKey = Other.SigningKey;
+		EncryptionKeys = Other.EncryptionKeys;
+		
+		if (Other.MasterEncryptionKey)
+		{
+			MasterEncryptionKey = EncryptionKeys.Find(Other.MasterEncryptionKey->Guid);
+		}
+		else
+		{
+			MasterEncryptionKey = nullptr;
+		}
+
+		return *this;
+	}
+
+	FKeyChain& operator=(FKeyChain&& Other)
+	{
+		SigningKey = Other.SigningKey;
+		EncryptionKeys = MoveTemp(Other.EncryptionKeys);
+		
+		if (Other.MasterEncryptionKey)
+		{
+			MasterEncryptionKey = EncryptionKeys.Find(Other.MasterEncryptionKey->Guid);
+		}
+		else
+		{
+			MasterEncryptionKey = nullptr;
+		}
+
+		Other.SigningKey = InvalidRSAKeyHandle;
+		Other.MasterEncryptionKey = nullptr;
+
+		return *this;
+	}
+
 	FRSAKeyHandle SigningKey = InvalidRSAKeyHandle;
 	TMap<FGuid, FNamedAESKey> EncryptionKeys;
 	const FNamedAESKey* MasterEncryptionKey = nullptr;
