add config option to use Oodle on files that are forced to zlib compression in pak
because they are needed in startup phase (ini,res,uplugin,etc.)
it's now fine to use Oodle there
this commit does not enable the new option so behavior stays the same for now

diff --git a/Engine/Config/BaseEngine.ini b/Engine/Config/BaseEngine.ini
--- a/Engine/Config/BaseEngine.ini
+++ b/Engine/Config/BaseEngine.ini
@@ -2458,6 +2458,8 @@ NumPreallocatedVertices=50000
 bTouchFallbackToMouse=true
 
 [Pak]
+; enable Oodle for files in the 'NotUsePluginCompression' list :
+;bDoUseOodleDespiteNoPluginCompression = true
 +ExtensionsToNotUsePluginCompression=uplugin
 +ExtensionsToNotUsePluginCompression=upluginmanifest
 +ExtensionsToNotUsePluginCompression=uproject
diff --git a/Engine/Source/Developer/PakFileUtilities/Private/PakFileUtilities.cpp b/Engine/Source/Developer/PakFileUtilities/Private/PakFileUtilities.cpp
--- a/Engine/Source/Developer/PakFileUtilities/Private/PakFileUtilities.cpp
+++ b/Engine/Source/Developer/PakFileUtilities/Private/PakFileUtilities.cpp
@@ -1884,20 +1884,36 @@ bool CreatePakFile(const TCHAR* Filename, TArray<FPakInputPair>& FilesToAdd, con
 		UE_LOG(LogPakFile, Display, TEXT("%s"), *FormatLogLine);
 	}
 
-	TArray<FString> ExtensionsToNotUsePluginCompression;
-	GConfig->GetArray(TEXT("Pak"), TEXT("ExtensionsToNotUsePluginCompression"), ExtensionsToNotUsePluginCompression, GEngineIni);
+	TSet<FString> NoPluginCompressionFileNames;
 	TSet<FString> NoPluginCompressionExtensions;
-	for (const FString& Ext : ExtensionsToNotUsePluginCompression)
+	
+	// Oodle is built into the Engine now and can be used to decode startup phase files (ini,res,uplugin)
+	//	 which used to be forced to Zlib
+	//	 set bDoUseOodleDespiteNoPluginCompression = true to let Oodle compress those files
+	bool bDoUseOodleDespiteNoPluginCompression = false;
+	GConfig->GetBool(TEXT("Pak"), TEXT("bDoUseOodleDespiteNoPluginCompression"), bDoUseOodleDespiteNoPluginCompression, GEngineIni);
+
+	if ( bDoUseOodleDespiteNoPluginCompression && CompressionFormatsAndNone[0] == NAME_Oodle )
 	{
-		NoPluginCompressionExtensions.Add(Ext);
+		// NoPluginCompression sets are empty
+
+		UE_LOG(LogPakFile, Display, TEXT("Oodle enabled on 'NoPluginCompression' files"));
 	}
-	
-	TArray<FString> FileNamesToNotUsePluginCompression;
-	GConfig->GetArray(TEXT("Pak"), TEXT("FileNamesToNotUsePluginCompression"), FileNamesToNotUsePluginCompression, GEngineIni);
-	TSet<FString> NoPluginCompressionFileNames;
-	for (const FString& FileName : FileNamesToNotUsePluginCompression)
+	else
 	{
-		NoPluginCompressionFileNames.Add(FileName);
+		TArray<FString> ExtensionsToNotUsePluginCompression;
+		GConfig->GetArray(TEXT("Pak"), TEXT("ExtensionsToNotUsePluginCompression"), ExtensionsToNotUsePluginCompression, GEngineIni);
+		for (const FString& Ext : ExtensionsToNotUsePluginCompression)
+		{
+			NoPluginCompressionExtensions.Add(Ext);
+		}
+	
+		TArray<FString> FileNamesToNotUsePluginCompression;
+		GConfig->GetArray(TEXT("Pak"), TEXT("FileNamesToNotUsePluginCompression"), FileNamesToNotUsePluginCompression, GEngineIni);
+		for (const FString& FileName : FileNamesToNotUsePluginCompression)
+		{
+			NoPluginCompressionFileNames.Add(FileName);
+		}
 	}
 
 	struct FAsyncCompressor
