- Allow UnrealPak to have remapped filenames, not just directories

diff --git a/Engine/Source/Developer/PakFileUtilities/Private/PakFileUtilities.cpp b/Engine/Source/Developer/PakFileUtilities/Private/PakFileUtilities.cpp
--- a/Engine/Source/Developer/PakFileUtilities/Private/PakFileUtilities.cpp
+++ b/Engine/Source/Developer/PakFileUtilities/Private/PakFileUtilities.cpp
@@ -1173,14 +1173,15 @@ void ProcessCommandLine(const TCHAR* CmdLine, const TArray<FString>& NonOptionAr
 				FPaths::NormalizeFilename(Input.Source);
 				if (SourceAndDest.Num() > 1)
 				{
-					Input.Dest = FPaths::GetPath(SourceAndDest[1]);
+					Input.Dest = SourceAndDest[1];
 				}
 				else
 				{
-					Input.Dest = FPaths::GetPath(Input.Source);
+					Input.Dest = Input.Source;
 				}
 				FPaths::NormalizeFilename(Input.Dest);
-				FPakFile::MakeDirectoryFromPath(Input.Dest);
+				FString DestPath = FPaths::GetPath(Input.Dest);
+				FPakFile::MakeDirectoryFromPath(DestPath);
 
 				//check for compression switches
 				for (int32 Index = 0; Index < Switches.Num(); ++Index)
@@ -1286,12 +1287,13 @@ void CollectFilesToAdd(TArray<FPakInputPair>& OutFilesToAdd, const TArray<FPakIn
 			TArray<FString> FoundFiles;
 			IFileManager::Get().FindFilesRecursive(FoundFiles, *Directory, *Filename, true, false);
 
+			FString DestDirectory = FPaths::GetPath(Input.Dest);
 			for (int32 FileIndex = 0; FileIndex < FoundFiles.Num(); FileIndex++)
 			{
 				FPakInputPair FileInput;
 				FileInput.Source = FoundFiles[FileIndex];
 				FPaths::MakeStandardFilename(FileInput.Source);
-				FileInput.Dest = FileInput.Source.Replace(*Directory, *Input.Dest, ESearchCase::IgnoreCase);
+				FileInput.Dest = FPaths::Combine(DestDirectory, FPaths::GetCleanFilename(FileInput.Source));
 				
 				uint64 FileOrder = OrderMap.GetFileOrder(FileInput.Dest, false, &FileInput.bIsInPrimaryOrder);
 				if(FileOrder != MAX_uint64)
@@ -1343,7 +1345,7 @@ void CollectFilesToAdd(TArray<FPakInputPair>& OutFilesToAdd, const TArray<FPakIn
 			FPakInputPair FileInput;
 			FileInput.Source = Input.Source;
 			FPaths::MakeStandardFilename(FileInput.Source);
-			FileInput.Dest = FileInput.Source.Replace(*Directory, *Input.Dest, ESearchCase::IgnoreCase);
+			FileInput.Dest = Input.Dest;
 			uint64 FileOrder = OrderMap.GetFileOrder(FileInput.Dest, CmdLineParameters.bFallbackOrderForNonUassetFiles, &FileInput.bIsInPrimaryOrder);
 			if (FileOrder != MAX_uint64)
 			{
