Fix whitespace from bad merge.
#rnx

diff --git a/Engine/Source/Developer/PakFileUtilities/Private/PakFileUtilities.cpp b/Engine/Source/Developer/PakFileUtilities/Private/PakFileUtilities.cpp
--- a/Engine/Source/Developer/PakFileUtilities/Private/PakFileUtilities.cpp
+++ b/Engine/Source/Developer/PakFileUtilities/Private/PakFileUtilities.cpp
@@ -131,7 +131,7 @@ public:
 	/** Divide into blocks and start compress asynchronously */
 	FMemoryCompressor(uint8* UncompressedBuffer, int64 UncompressedSize, FName Format, int32 CompressionBlockSize) :
 		Index(0)
-			{
+	{
 		// Divide into blocks and start compression async tasks.
 		// These blocks must be as same as followed CompressMemory callings.
 		int64 UncompressedBytes = 0;
@@ -757,8 +757,8 @@ bool FCompressedFileBuffer::CompressFileToWorkingBuffer(const FPakInputPair& InF
 
 		// CompressMemoryBound truncates its size argument to 32bits, so we can not use (possibly > 32-bit) FileSize directly to calculate required buffer space
 		int32 MaxCompressedBufferSize = Align(FCompression::CompressMemoryBound(CompressionMethod, CompressionBlockSize, COMPRESS_NoFlags), FAES::AESBlockSize);
-        int32 CompressionBufferRemainder = Align(FCompression::CompressMemoryBound(CompressionMethod, int32(UncompressedSize % CompressionBlockSize), COMPRESS_NoFlags), FAES::AESBlockSize);
-        EnsureBufferSpace(MaxCompressedBufferSize * (UncompressedSize / CompressionBlockSize) + CompressionBufferRemainder);
+		int32 CompressionBufferRemainder = Align(FCompression::CompressMemoryBound(CompressionMethod, int32(UncompressedSize % CompressionBlockSize), COMPRESS_NoFlags), FAES::AESBlockSize);
+		EnsureBufferSpace(MaxCompressedBufferSize * (UncompressedSize / CompressionBlockSize) + CompressionBufferRemainder);
 
 		TotalCompressedSize = 0;
 		int64 UncompressedBytes = 0;
@@ -2130,22 +2130,22 @@ bool CreatePakFile(const TCHAR* Filename, TArray<FPakInputPair>& FilesToAdd, con
 		for (int32 LaunchFileIndex = NextLaunchIndex; LaunchFileIndex < LaunchFileIndexEnd; LaunchFileIndex++)
 		{
 			AsyncCompressors[LaunchFileIndex].Init(&FilesToAdd[LaunchFileIndex], CmdLineParameters, &NoPluginCompressionExtensions, &NoPluginCompressionFileNames);
-			
-		if (bRunAsync)
-		{
-				if (FilesToAdd[LaunchFileIndex].bNeedsCompression)
+
+			if (bRunAsync)
 			{
+				if (FilesToAdd[LaunchFileIndex].bNeedsCompression)
+				{
 					(new FAutoDeleteAsyncTask<FRunCompressionTask>(&AsyncCompressors[LaunchFileIndex]))->StartBackgroundTask();
-			}
-			else
-			{
-				// call compress function inline 
-				// it won't do anything except for initialize some internal variables used in the non compressed path
-				// we don't want to pass these to a different thread as they may cause congestion with legitimate tasks
+				}
+				else
+				{
+					// call compress function inline 
+					// it won't do anything except for initialize some internal variables used in the non compressed path
+					// we don't want to pass these to a different thread as they may cause congestion with legitimate tasks
 					AsyncCompressors[LaunchFileIndex].Compress();
+				}
 			}
 		}
-	}
 		NextLaunchIndex = LaunchFileIndexEnd;
 
 		bool bDeleted = FilesToAdd[FileIndex].bIsDeleteRecord;
@@ -2208,33 +2208,33 @@ bool CreatePakFile(const TCHAR* Filename, TArray<FPakInputPair>& FilesToAdd, con
 				(CmdLineParameters.bAlignFilesLargerThanBlock || RealFileSize <= CmdLineParameters.FileSystemBlockSize) &&
 				(NewEntryOffset / CmdLineParameters.FileSystemBlockSize) != ((NewEntryOffset + RealFileSize - 1) / CmdLineParameters.FileSystemBlockSize)) // File crosses a block boundary
 			{
-					int64 OldOffset = NewEntryOffset;
-					NewEntryOffset = AlignArbitrary(NewEntryOffset, CmdLineParameters.FileSystemBlockSize);
-					int64 PaddingRequired = NewEntryOffset - OldOffset;
+				int64 OldOffset = NewEntryOffset;
+				NewEntryOffset = AlignArbitrary(NewEntryOffset, CmdLineParameters.FileSystemBlockSize);
+				int64 PaddingRequired = NewEntryOffset - OldOffset;
 
-					check(PaddingRequired >= 0);
-					check(PaddingRequired < RealFileSize);
-					if (PaddingRequired > 0)
+				check(PaddingRequired >= 0);
+				check(PaddingRequired < RealFileSize);
+				if (PaddingRequired > 0)
+				{
+					// If we don't already have a padding buffer, create one
+					if (PaddingBuffer == nullptr)
 					{
-						// If we don't already have a padding buffer, create one
-						if (PaddingBuffer == nullptr)
-						{
-							PaddingBufferSize = 64 * 1024;
-							PaddingBuffer = (uint8*)FMemory::Malloc(PaddingBufferSize);
-							FMemory::Memset(PaddingBuffer, 0, PaddingBufferSize);
-						}
-
-						UE_LOG(LogPakFile, Verbose, TEXT("%14llu - %14llu : %14llu padding."), PakFileHandle->Tell(), PakFileHandle->Tell() + PaddingRequired, PaddingRequired);
-						while (PaddingRequired > 0)
-						{
-							int64 AmountToWrite = FMath::Min(PaddingRequired, PaddingBufferSize);
-							PakFileHandle->Serialize(PaddingBuffer, AmountToWrite);
-							PaddingRequired -= AmountToWrite;
-						}
+						PaddingBufferSize = 64 * 1024;
+						PaddingBuffer = (uint8*)FMemory::Malloc(PaddingBufferSize);
+						FMemory::Memset(PaddingBuffer, 0, PaddingBufferSize);
+					}
 
-						check(PakFileHandle->Tell() == NewEntryOffset);
+					UE_LOG(LogPakFile, Verbose, TEXT("%14llu - %14llu : %14llu padding."), PakFileHandle->Tell(), PakFileHandle->Tell() + PaddingRequired, PaddingRequired);
+					while (PaddingRequired > 0)
+					{
+						int64 AmountToWrite = FMath::Min(PaddingRequired, PaddingBufferSize);
+						PakFileHandle->Serialize(PaddingBuffer, AmountToWrite);
+						PaddingRequired -= AmountToWrite;
 					}
+
+					check(PakFileHandle->Tell() == NewEntryOffset);
 				}
+			}
 
 			// Align bulk data
 			if (bIsMappedBulk && CmdLineParameters.AlignForMemoryMapping > 0 && OriginalFileSize != INDEX_NONE && !bDeleted)
@@ -3971,31 +3971,31 @@ bool GenerateHashesFromPak(const TCHAR* InPakFilename, const TCHAR* InDestPakFil
 				}
 				else
 				{
-				    PakReader->Seek(Entry.Offset);
-				    uint32 SerializedCrcTest = 0;
-				    FPakEntry EntryInfo;
-				    EntryInfo.Serialize(PakReader.GetArchive(), PakFile.GetInfo().Version);
-				    if (EntryInfo.IndexDataEquals(Entry))
-				    {
-					    // TUniquePtr<FArchive> FileHandle(IFileManager::Get().CreateFileWriter(*DestFilename));
-					    TArray<uint8> Bytes;
-					    FMemoryWriter MemoryFile(Bytes);
-					    FArchive* FileHandle = &MemoryFile;
-					    // if (FileHandle.IsValid())
-					    {
+					PakReader->Seek(Entry.Offset);
+					uint32 SerializedCrcTest = 0;
+					FPakEntry EntryInfo;
+					EntryInfo.Serialize(PakReader.GetArchive(), PakFile.GetInfo().Version);
+					if (EntryInfo.IndexDataEquals(Entry))
+					{
+						// TUniquePtr<FArchive> FileHandle(IFileManager::Get().CreateFileWriter(*DestFilename));
+						TArray<uint8> Bytes;
+						FMemoryWriter MemoryFile(Bytes);
+						FArchive* FileHandle = &MemoryFile;
+						// if (FileHandle.IsValid())
+						{
 							if (Entry.CompressionMethodIndex == 0)
-						    {
-							    BufferedCopyFile(*FileHandle, PakReader.GetArchive(), PakFile, Entry, Buffer, BufferSize, InKeyChain);
-						    }
-						    else
-						    {
-							    UncompressCopyFile(*FileHandle, PakReader.GetArchive(), Entry, PersistantCompressionBuffer, CompressionBufferSize, InKeyChain, PakFile);
-						    }
-    
-						    UE_LOG(LogPakFile, Display, TEXT("Generated hash for \"%s\""), *FullFilename);
-						    GenerateHashForFile(Bytes.GetData(), Bytes.Num(), FileHash);
-						    FileHash.PatchIndex = PakPriority;
-						    FileHash.bIsDeleteRecord = false;
+							{
+								BufferedCopyFile(*FileHandle, PakReader.GetArchive(), PakFile, Entry, Buffer, BufferSize, InKeyChain);
+							}
+							else
+							{
+								UncompressCopyFile(*FileHandle, PakReader.GetArchive(), Entry, PersistantCompressionBuffer, CompressionBufferSize, InKeyChain, PakFile);
+							}
+
+							UE_LOG(LogPakFile, Display, TEXT("Generated hash for \"%s\""), *FullFilename);
+							GenerateHashForFile(Bytes.GetData(), Bytes.Num(), FileHash);
+							FileHash.PatchIndex = PakPriority;
+							FileHash.bIsDeleteRecord = false;
 							FileHash.bForceInclude = false;
 							bEntryValid = true;
 						}
