Removed runtime dependencies on DerivedDataCache

#rnx

diff --git a/Engine/Source/Developer/IoStoreUtilities/IoStoreUtilities.Build.cs b/Engine/Source/Developer/IoStoreUtilities/IoStoreUtilities.Build.cs
--- a/Engine/Source/Developer/IoStoreUtilities/IoStoreUtilities.Build.cs
+++ b/Engine/Source/Developer/IoStoreUtilities/IoStoreUtilities.Build.cs
@@ -15,7 +15,6 @@ public class IoStoreUtilities : ModuleRules
             "CoreUObject",
             "AssetRegistry",
 			"Zen",
-			"DerivedDataCache",
 			"RenderCore",
 			"libcurl",
         });
diff --git a/Engine/Source/Developer/PakFileUtilities/PakFileUtilities.Build.cs b/Engine/Source/Developer/PakFileUtilities/PakFileUtilities.Build.cs
--- a/Engine/Source/Developer/PakFileUtilities/PakFileUtilities.Build.cs
+++ b/Engine/Source/Developer/PakFileUtilities/PakFileUtilities.Build.cs
@@ -7,20 +7,22 @@ public class PakFileUtilities : ModuleRules
 	public PakFileUtilities(ReadOnlyTargetRules Target) : base(Target)
 	{
 		PrivateDependencyModuleNames.AddRange(new string[] {
-            "Core",
-            "PakFile",
-            "Json",
-            "Projects",
-            "RSA",
-            "IoStoreUtilities",
-        });
+			"Core",
+			"PakFile",
+			"Json",
+			"Projects",
+			"RSA",
+			"IoStoreUtilities",
+		});
 
-        PrivateIncludePathModuleNames.AddRange(new string[] {
-            "DerivedDataCache",
-            "Json",
-        });
+		PrivateIncludePathModuleNames.AddRange(new string[] {
+			"DerivedDataCache",
+			"Json",
+		});
 
-		DynamicallyLoadedModuleNames.Add("DerivedDataCache");
-		
+		if (Target.bBuildWithEditorOnlyData)
+		{
+			DynamicallyLoadedModuleNames.Add("DerivedDataCache");
+		}
 	}
 }
diff --git a/Engine/Source/Developer/PakFileUtilities/Private/PakFileUtilities.cpp b/Engine/Source/Developer/PakFileUtilities/Private/PakFileUtilities.cpp
--- a/Engine/Source/Developer/PakFileUtilities/Private/PakFileUtilities.cpp
+++ b/Engine/Source/Developer/PakFileUtilities/Private/PakFileUtilities.cpp
@@ -2108,7 +2108,9 @@ bool CreatePakFile(const TCHAR* Filename, TArray<FPakInputPair>& FilesToAdd, con
 
 	const bool bRunAsync = CmdLineParameters.bAsyncCompression;
 
+#if USE_DDC_FOR_COMPRESSED_FILES
 	GetDerivedDataCacheRef();
+#endif
 	
 	TArray<FFileRegion> AllFileRegions;
 	
@@ -5528,9 +5530,9 @@ bool ExecuteUnrealPak(const TCHAR* CmdLine)
 			IFileManager::Get().DeleteDirectory(*OutputPath, false, true);
 		}
 
+#if USE_DDC_FOR_COMPRESSED_FILES
 		GetDerivedDataCacheRef().WaitForQuiescence(true);
-
-		
+#endif
 
 		return bResult;
 	}
