UnrealPak: Add support for creating multiple pak files at once instead of running a separate process for each

diff --git a/Engine/Source/Developer/PakFileUtilities/Private/PakFileUtilities.cpp b/Engine/Source/Developer/PakFileUtilities/Private/PakFileUtilities.cpp
--- a/Engine/Source/Developer/PakFileUtilities/Private/PakFileUtilities.cpp
+++ b/Engine/Source/Developer/PakFileUtilities/Private/PakFileUtilities.cpp
@@ -30,6 +30,9 @@
 #include "Misc/KeyChainUtilities.h"
 #include "IoStoreUtilities.h"
 #include "Interfaces/IPluginManager.h"
+#include "Containers/SpscQueue.h"
+#include "Async/Async.h"
+#include "Async/Future.h"
 
 IMPLEMENT_MODULE(FDefaultModuleImpl, PakFileUtilities);
 
@@ -433,7 +436,6 @@ struct FPakCommandLineParameters
 		, bSign(false)
 		, bPatchCompatibilityMode421(false)
 		, bFallbackOrderForNonUassetFiles(false)
-		, bAsyncCompression(false)
 		, bAlignFilesLargerThanBlock(false)
 		, bForceCompress(false)
 		, bFileRegions(false)
@@ -457,7 +459,6 @@ struct FPakCommandLineParameters
 	bool bSign;
 	bool bPatchCompatibilityMode421;
 	bool bFallbackOrderForNonUassetFiles;
-	bool bAsyncCompression;
 	bool bAlignFilesLargerThanBlock;	// Align files that are larger than block size
 
 	bool bForceCompress; // Force all files that request compression to be compressed, even if that results in a larger file size DEPRECATED
@@ -562,6 +563,8 @@ struct FCompressedFileBuffer
 		}
 	}
 
+	bool ReadUncompressedFileToWorkingBuffer(const FPakInputPair& InFile);
+	
 	bool CompressFileToWorkingBuffer(const FPakInputPair& InFile, FName CompressionMethod, const int32 CompressionBlockSize);
 
 	void SerializeDDCData(FArchive &Ar)
@@ -623,14 +626,14 @@ bool ReadSizeParam(const TCHAR* CmdLine, const TCHAR* ParamStr, T& SizeOut)
 }
 
 
-FString GetLongestPath(TArray<FPakInputPair>& FilesToAdd)
+FString GetLongestPath(const TArray<FPakInputPair>& FilesToAdd)
 {
 	FString LongestPath;
 	int32 MaxNumDirectories = 0;
 
 	for (int32 FileIndex = 0; FileIndex < FilesToAdd.Num(); FileIndex++)
 	{
-		FString& Filename = FilesToAdd[FileIndex].Dest;
+		const FString& Filename = FilesToAdd[FileIndex].Dest;
 		int32 NumDirectories = 0;
 		for (int32 Index = 0; Index < Filename.Len(); Index++)
 		{
@@ -648,7 +651,7 @@ FString GetLongestPath(TArray<FPakInputPair>& FilesToAdd)
 	return FPaths::GetPath(LongestPath) + TEXT("/");
 }
 
-FString GetCommonRootPath(TArray<FPakInputPair>& FilesToAdd)
+FString GetCommonRootPath(const TArray<FPakInputPair>& FilesToAdd)
 {
 	FString Root = GetLongestPath(FilesToAdd);
 	for (int32 FileIndex = 0; FileIndex < FilesToAdd.Num() && Root.Len(); FileIndex++)
@@ -698,6 +701,23 @@ FString FCompressedFileBuffer::GetDDCKeyString(const uint8* UncompressedFile, co
 	return FDerivedDataCacheInterface::BuildCacheKey(TEXT("PAKCOMPRESS_"), PAKCOMPRESS_DERIVEDDATA_VER, *KeyString);
 }
 
+bool FCompressedFileBuffer::ReadUncompressedFileToWorkingBuffer(const FPakInputPair& InFile)
+{
+	TUniquePtr<FArchive> FileHandle(IFileManager::Get().CreateFileReader(*InFile.Source));
+	if (!FileHandle)
+	{
+		OriginalSize = -1;
+		return false;
+	}
+	OriginalSize = FileHandle->TotalSize();
+	const int64 PaddedEncryptedFileSize = Align(OriginalSize, FAES::AESBlockSize);
+	CompressedBuffer = MakeUnique<uint8[]>(PaddedEncryptedFileSize);
+	CompressedBufferSize = PaddedEncryptedFileSize;
+	// Load to buffer
+	FileHandle->Serialize(CompressedBuffer.Get(), OriginalSize);
+	return true;
+}
+
 bool FCompressedFileBuffer::CompressFileToWorkingBuffer(const FPakInputPair& InFile, FName CompressionMethod, const int32 CompressionBlockSize)
 {
 	TUniquePtr<FArchive> FileHandle(IFileManager::Get().CreateFileReader(*InFile.Source));
@@ -815,66 +835,58 @@ bool FCompressedFileBuffer::CompressFileToWorkingBuffer(const FPakInputPair& InF
 	return true;
 }
 
-bool PrepareCopyFileToPak(const FString& InMountPoint, const FPakInputPair& InFile, uint8*& InOutPersistentBuffer, int64& InOutBufferSize, FPakEntryPair& OutNewEntry, uint8*& OutDataToWrite, int64& OutSizeToWrite, const FKeyChain& InKeyChain, TArray<FFileRegion>* OutFileRegions)
-{	
-	TUniquePtr<FArchive> FileHandle(IFileManager::Get().CreateFileReader(*InFile.Source));
-	bool bFileExists = FileHandle.IsValid();
-	if (bFileExists)
+bool PrepareCopyFileToPak(const FString& InMountPoint, const FPakInputPair& InFile, const FCompressedFileBuffer& UncompressedFile, FPakEntryPair& OutNewEntry, uint8*& OutDataToWrite, int64& OutSizeToWrite, const FKeyChain& InKeyChain, TArray<FFileRegion>* OutFileRegions)
+{
+	const int64 FileSize = UncompressedFile.OriginalSize;
+	const int64 PaddedEncryptedFileSize = Align(FileSize, FAES::AESBlockSize); 
+	if (FileSize < 0)
 	{
-		const int64 FileSize = FileHandle->TotalSize();
-		const int64 PaddedEncryptedFileSize = Align(FileSize, FAES::AESBlockSize); 
-		OutNewEntry.Filename = InFile.Dest.Mid(InMountPoint.Len());
-		OutNewEntry.Info.Offset = 0; // Don't serialize offsets here.
-		OutNewEntry.Info.Size = FileSize;
-		OutNewEntry.Info.UncompressedSize = FileSize;
-		OutNewEntry.Info.CompressionMethodIndex = 0;
-		OutNewEntry.Info.SetEncrypted( InFile.bNeedEncryption );
-		OutNewEntry.Info.SetDeleteRecord(false);
-
-		if (InOutBufferSize < PaddedEncryptedFileSize)
-		{
-			InOutPersistentBuffer = (uint8*)FMemory::Realloc(InOutPersistentBuffer, PaddedEncryptedFileSize);
-			InOutBufferSize = FileSize;
-		}
-
-		// Load to buffer
-		FileHandle->Serialize(InOutPersistentBuffer, FileSize);
+		return false;
+	}
+	check(UncompressedFile.CompressedBufferSize == PaddedEncryptedFileSize);
+	OutNewEntry.Filename = InFile.Dest.Mid(InMountPoint.Len());
+	OutNewEntry.Info.Offset = 0; // Don't serialize offsets here.
+	OutNewEntry.Info.Size = FileSize;
+	OutNewEntry.Info.UncompressedSize = FileSize;
+	OutNewEntry.Info.CompressionMethodIndex = 0;
+	OutNewEntry.Info.SetEncrypted( InFile.bNeedEncryption );
+	OutNewEntry.Info.SetDeleteRecord(false);
 
+	{
+		OutSizeToWrite = FileSize;
+		uint8* UncompressedBuffer = UncompressedFile.CompressedBuffer.Get();
+		if (InFile.bNeedEncryption)
 		{
-			OutSizeToWrite = FileSize;
-			if (InFile.bNeedEncryption)
+			for(int64 FillIndex = FileSize; FillIndex < PaddedEncryptedFileSize; ++FillIndex)
 			{
-				for(int64 FillIndex=FileSize; FillIndex < PaddedEncryptedFileSize && InFile.bNeedEncryption; ++FillIndex)
-				{
-					// Fill the trailing buffer with bytes from file. Note that this is now from a fixed location
-					// rather than a random one so that we produce deterministic results
-					InOutPersistentBuffer[FillIndex] = InOutPersistentBuffer[(FillIndex - FileSize)%FileSize];
-				}
-
-				//Encrypt the buffer before writing it to disk
-				check(InKeyChain.MasterEncryptionKey);
-				FAES::EncryptData(InOutPersistentBuffer, PaddedEncryptedFileSize, InKeyChain.MasterEncryptionKey->Key);
-				// Update the size to be written
-				OutSizeToWrite = PaddedEncryptedFileSize;
-				OutNewEntry.Info.SetEncrypted( true );
+				// Fill the trailing buffer with bytes from file. Note that this is now from a fixed location
+				// rather than a random one so that we produce deterministic results
+				UncompressedBuffer[FillIndex] = UncompressedBuffer[(FillIndex - FileSize)%FileSize];
 			}
 
-			// Calculate the buffer hash value
-			FSHA1::HashBuffer(InOutPersistentBuffer,FileSize,OutNewEntry.Info.Hash);			
-			OutDataToWrite = InOutPersistentBuffer;
+			//Encrypt the buffer before writing it to disk
+			check(InKeyChain.MasterEncryptionKey);
+			FAES::EncryptData(UncompressedBuffer, PaddedEncryptedFileSize, InKeyChain.MasterEncryptionKey->Key);
+			// Update the size to be written
+			OutSizeToWrite = PaddedEncryptedFileSize;
+			OutNewEntry.Info.SetEncrypted( true );
 		}
 
-		if (OutFileRegions)
+		// Calculate the buffer hash value
+		FSHA1::HashBuffer(UncompressedBuffer, FileSize, OutNewEntry.Info.Hash);
+		OutDataToWrite = UncompressedBuffer;
+	}
+
+	if (OutFileRegions)
+	{
+		// Read the matching regions file, if it exists.
+		TUniquePtr<FArchive> RegionsFile(IFileManager::Get().CreateFileReader(*(InFile.Source + FFileRegion::RegionsFileExtension)));
+		if (RegionsFile.IsValid())
 		{
-			// Read the matching regions file, if it exists.
-			TUniquePtr<FArchive> RegionsFile(IFileManager::Get().CreateFileReader(*(InFile.Source + FFileRegion::RegionsFileExtension)));
-			if (RegionsFile.IsValid())
-			{
-				FFileRegion::SerializeFileRegions(*RegionsFile.Get(), *OutFileRegions);
-			}
+			FFileRegion::SerializeFileRegions(*RegionsFile.Get(), *OutFileRegions);
 		}
 	}
-	return bFileExists;
+	return true;
 }
 
 void FinalizeCopyCompressedFileToPak(FPakInfo& InPakInfo, const FCompressedFileBuffer& CompressedFile, FPakEntryPair& OutNewEntry)
@@ -939,27 +951,9 @@ void PrepareDeleteRecordForPak(const FString& InMountPoint, const FPakInputPair
 	OutNewEntry.Info.SetDeleteRecord(true);
 }
 
-
-static void CommandLineParseHelper(const TCHAR* InCmdLine, TArray<FString>& Tokens, TArray<FString>& Switches)
-{
-	FString NextToken;
-	while(FParse::Token(InCmdLine,NextToken,false))
-	{
-		if((**NextToken == TCHAR('-')))
-		{
-			new(Switches)FString(NextToken.Mid(1));
-		}
-		else
-		{
-			new(Tokens)FString(NextToken);
-		}
-	}
-}
-
-void ProcessCommandLine(const TCHAR* CmdLine, const TArray<FString>& NonOptionArguments, TArray<FPakInputPair>& Entries, FPakCommandLineParameters& CmdLineParameters)
+void ProcessCommonCommandLine(const TCHAR* CmdLine, FPakCommandLineParameters& CmdLineParameters)
 {
 	// List of all items to add to pak file
-	FString ResponseFile;
 	FString ClusterSizeString;
 
 	if (FParse::Param(CmdLine, TEXT("patchcompatibilitymode421")))
@@ -967,11 +961,6 @@ void ProcessCommandLine(const TCHAR* CmdLine, const TArray<FString>& NonOptionAr
 		CmdLineParameters.bPatchCompatibilityMode421 = true;
 	}
 
-	if (FParse::Param(CmdLine, TEXT("asynccompression")))
-	{
-		CmdLineParameters.bAsyncCompression = true;
-	}
-
 	if (FParse::Param(CmdLine, TEXT("fallbackOrderForNonUassetFiles")))
 	{
 		CmdLineParameters.bFallbackOrderForNonUassetFiles = true;
@@ -1077,6 +1066,21 @@ void ProcessCommandLine(const TCHAR* CmdLine, const TArray<FString>& NonOptionAr
 		CmdLineParameters.CompressionFormats.AddUnique(NAME_Zlib);
 	}
 
+	FParse::Value(CmdLine, TEXT("-patchSeekOptMaxInflationPercent="), CmdLineParameters.SeekOptParams.MaxInflationPercent);
+	ReadSizeParam(CmdLine, TEXT("-patchSeekOptMaxGapSize="), CmdLineParameters.SeekOptParams.MaxGapSize);
+	FParse::Value(CmdLine, TEXT("-patchSeekOptMaxAdjacentOrderDiff="), CmdLineParameters.SeekOptParams.MaxAdjacentOrderDiff);
+
+	// For legacy reasons, if we specify a max gap size without a mode, we default to OnePass
+	if (CmdLineParameters.SeekOptParams.MaxGapSize > 0)
+	{
+		CmdLineParameters.SeekOptParams.Mode = ESeekOptMode::OnePass;
+	}
+	FParse::Value(CmdLine, TEXT("-patchSeekOptMode="), (int32&)CmdLineParameters.SeekOptParams.Mode);
+}
+
+void ProcessPakFileSpecificCommandLine(const TCHAR* CmdLine, const TArray<FString>& NonOptionArguments, TArray<FPakInputPair>& Entries, FPakCommandLineParameters& CmdLineParameters)
+{
+	FString ResponseFile;
 	if (FParse::Value(CmdLine, TEXT("-create="), ResponseFile))
 	{
 		CmdLineParameters.GeneratePatch = FParse::Value(CmdLine, TEXT("-generatepatch="), CmdLineParameters.SourcePatchPakFilename);
@@ -1092,22 +1096,6 @@ void ProcessCommandLine(const TCHAR* CmdLine, const TArray<FString>& NonOptionAr
 
 		bool bEncrypt = FParse::Param(CmdLine, TEXT("encrypt"));		
 
-		if (CmdLineParameters.GeneratePatch)
-		{
-			FParse::Value(CmdLine, TEXT("-patchSeekOptMaxInflationPercent="), CmdLineParameters.SeekOptParams.MaxInflationPercent);
-			ReadSizeParam(CmdLine, TEXT("-patchSeekOptMaxGapSize="), CmdLineParameters.SeekOptParams.MaxGapSize);
-			FParse::Value(CmdLine, TEXT("-patchSeekOptMaxAdjacentOrderDiff="), CmdLineParameters.SeekOptParams.MaxAdjacentOrderDiff);
-
-			// For legacy reasons, if we specify a max gap size without a mode, we default to OnePass
-			if (CmdLineParameters.SeekOptParams.MaxGapSize > 0)
-			{
-				CmdLineParameters.SeekOptParams.Mode = ESeekOptMode::OnePass;
-			}
-			FParse::Value(CmdLine, TEXT("-patchSeekOptMode="), (int32&)CmdLineParameters.SeekOptParams.Mode);
-
-			
-		}
-
 		// if the response file is a pak file, then this is the pak file we want to use as the source
 		if (ResponseFile.EndsWith(TEXT(".pak"), ESearchCase::IgnoreCase) && CmdLineParameters.GeneratePatch)
 		{
@@ -1139,6 +1127,7 @@ void ProcessCommandLine(const TCHAR* CmdLine, const TArray<FString>& NonOptionAr
 			}
 			else
 			{
+				TRACE_CPUPROFILER_EVENT_SCOPE(LoadResponseFile);
 				FString Text;
 				UE_LOG(LogPakFile, Display, TEXT("Loading response file %s"), *ResponseFile);
 				if (FFileHelper::LoadFileToString(Text, *ResponseFile))
@@ -1154,60 +1143,85 @@ void ProcessCommandLine(const TCHAR* CmdLine, const TArray<FString>& NonOptionAr
 				}
 			}
 
+			TRACE_CPUPROFILER_EVENT_SCOPE(AddEntries);
+			
+			Entries.Reserve(Lines.Num());
+			FString NextToken;
 			for (int32 EntryIndex = 0; EntryIndex < Lines.Num(); EntryIndex++)
 			{
-				TArray<FString> SourceAndDest;
-				TArray<FString> Switches;
+				FPakInputPair& Input = Entries.AddDefaulted_GetRef();
 				if (bParseLines)
 				{
-					Lines[EntryIndex].TrimStartInline();
-					CommandLineParseHelper(*Lines[EntryIndex], SourceAndDest, Switches);
+					TRACE_CPUPROFILER_EVENT_SCOPE(CommandLineParseHelper);
+					bool bHasParsedSource = false;
+					bool bHasParsedDest = false;
+					const TCHAR* LinePtr = *Lines[EntryIndex];
+					while (FParse::Token(LinePtr, NextToken, false))
+					{
+						FStringView TokenView(NextToken);
+						if (TokenView[0] == TCHAR('-'))
+						{
+							FStringView Switch = TokenView.Mid(1);
+							if (Switch == TEXT("compress"))
+							{
+								Input.bNeedsCompression = true;
+							}
+							else if (Switch == TEXT("encrypt"))
+							{
+								Input.bNeedEncryption = true;
+							}
+							else if (Switch == TEXT("delete"))
+							{
+								Input.bIsDeleteRecord = true;
+							}
+						}
+						else
+						{
+							if (!bHasParsedSource)
+							{
+								Input.Source = MoveTemp(NextToken);
+								bHasParsedSource = true;
+							}
+							else if (!bHasParsedDest)
+							{
+								Input.Dest = MoveTemp(NextToken);
+								bHasParsedDest = true;
+							}
+						}
+					}
+					if (!bHasParsedSource)
+					{
+						Entries.Pop(false);
+						continue;
+					}
+					
+					if (!bHasParsedDest || Input.bIsDeleteRecord)
+					{
+						Input.Dest = Input.Source;
+					}
 				}
 				else
 				{
-					SourceAndDest.Add(Lines[EntryIndex]);
-				}
-				if (SourceAndDest.Num() == 0)
-				{
-					continue;
+					Input.Source = MoveTemp(Lines[EntryIndex]);
+					Input.Dest = Input.Source;
 				}
-				FPakInputPair Input;
 
-				Input.Source = SourceAndDest[0];
-				FPaths::NormalizeFilename(Input.Source);
-				if (SourceAndDest.Num() > 1)
+				FString DestPath;
 				{
-					Input.Dest = SourceAndDest[1];
+					TRACE_CPUPROFILER_EVENT_SCOPE(OtherPathStuff);
+					FPaths::NormalizeFilename(Input.Source);
+					FPaths::NormalizeFilename(Input.Dest);
+					DestPath = FPaths::GetPath(Input.Dest);
 				}
-				else
 				{
-					Input.Dest = Input.Source;
+					TRACE_CPUPROFILER_EVENT_SCOPE(MakeDirectoryFromPath);
+					FPakFile::MakeDirectoryFromPath(DestPath);
 				}
-				FPaths::NormalizeFilename(Input.Dest);
-				FString DestPath = FPaths::GetPath(Input.Dest);
-				FPakFile::MakeDirectoryFromPath(DestPath);
 
-				//check for compression switches
-				for (int32 Index = 0; Index < Switches.Num(); ++Index)
-				{
-					if (Switches[Index] == TEXT("compress"))
-					{
-						Input.bNeedsCompression = true;
-					}
-					if (Switches[Index] == TEXT("encrypt"))
-					{
-						Input.bNeedEncryption = true;
-					}
-					if (Switches[Index] == TEXT("delete"))
-					{
-						Input.bIsDeleteRecord = true;
-						Input.Dest = Input.Source;
-					}
-				}
 				Input.bNeedsCompression |= bCompress;
 				Input.bNeedEncryption |= bEncrypt;
 
-				UE_LOG(LogPakFile, Log, TEXT("Added file Source: %s Dest: %s"), *Input.Source, *Input.Dest);
+				UE_LOG(LogPakFile, Verbose, TEXT("Added file Source: %s Dest: %s"), *Input.Source, *Input.Dest);
 
 				bool bIsMappedBulk = Input.Source.EndsWith(TEXT(".m.ubulk"));
 				if (bIsMappedBulk && CmdLineParameters.AlignForMemoryMapping > 0 && Input.bNeedsCompression
@@ -1215,9 +1229,8 @@ void ProcessCommandLine(const TCHAR* CmdLine, const TArray<FString>& NonOptionAr
 				{
 					// no compression for bulk aligned files because they are memory mapped
 					Input.bNeedsCompression = false;
-					UE_LOG(LogPakFile, Log, TEXT("Stripped compression from %s for memory mapping."), *Input.Dest);
+					UE_LOG(LogPakFile, Verbose, TEXT("Stripped compression from %s for memory mapping."), *Input.Dest);
 				}
-				Entries.Add(Input);
 			}
 		}
 
@@ -1255,8 +1268,15 @@ void ProcessCommandLine(const TCHAR* CmdLine, const TArray<FString>& NonOptionAr
 	}
 }
 
+void ProcessCommandLine(const TCHAR* CmdLine, const TArray<FString>& NonOptionArguments, TArray<FPakInputPair>& Entries, FPakCommandLineParameters& CmdLineParameters)
+{
+	ProcessCommonCommandLine(CmdLine, CmdLineParameters);
+	ProcessPakFileSpecificCommandLine(CmdLine, NonOptionArguments, Entries, CmdLineParameters);
+}
+
 void CollectFilesToAdd(TArray<FPakInputPair>& OutFilesToAdd, const TArray<FPakInputPair>& InEntries, const FPakOrderMap& OrderMap, const FPakCommandLineParameters& CmdLineParameters)
 {
+	TRACE_CPUPROFILER_EVENT_SCOPE(CollectFilesToAdd);
 	UE_LOG(LogPakFile, Display, TEXT("Collecting files to add to pak file..."));
 	const double StartTime = FPlatformTime::Seconds();
 
@@ -1788,91 +1808,110 @@ void VerifyIndexesMatch(TArray<FPakEntryPair>& EntryList, FPakFile::FDirectoryIn
 	}
 };
 
-TAtomic<int64> GTotalFilesWithPoorForcedCompression(0);
-TAtomic<int64> GTotalExtraMemoryForPoorForcedCompression(0);
-
-bool CreatePakFile(const TCHAR* Filename, TArray<FPakInputPair>& FilesToAdd, const FPakCommandLineParameters& CmdLineParameters, const FKeyChain& InKeyChain)
+class FPakWriterContext
 {
-	const double StartTime = FPlatformTime::Seconds();
+public:
+	bool Initialize(const FPakCommandLineParameters& InCmdLineParameters);
+	bool AddPakFile(const TCHAR* Filename, const TArray<FPakInputPair>& FilesToAdd, const FKeyChain& InKeyChain);
+	bool Flush();
 
-	// Create Pak
-	TUniquePtr<FArchive> PakFileHandle(CreatePakWriter(Filename, InKeyChain, CmdLineParameters.bSign));
-	if (!PakFileHandle)
-	{
-		UE_LOG(LogPakFile, Error, TEXT("Unable to create pak file \"%s\"."), Filename);
-		return false;
-	}
+private:
+	struct FOutputPakFile;
 
-	TUniquePtr<FArchive> PakFileRegionsHandle;
-	if (CmdLineParameters.bFileRegions)
+	struct FOutputPakFileEntry
 	{
-		FString RegionsFilename = FString(Filename) + FFileRegion::RegionsFileExtension;
-		PakFileRegionsHandle.Reset(IFileManager::Get().CreateFileWriter(*RegionsFilename));
-		if (!PakFileRegionsHandle)
-		{
-			UE_LOG(LogPakFile, Error, TEXT("Unable to create pak regions file \"%s\"."), *RegionsFilename);
-			return false;
-		}
-	}
-
-	FPakInfo Info;
-	Info.bEncryptedIndex = (InKeyChain.MasterEncryptionKey && CmdLineParameters.EncryptIndex);
-	Info.EncryptionKeyGuid = InKeyChain.MasterEncryptionKey ? InKeyChain.MasterEncryptionKey->Guid : FGuid();
-
+		FOutputPakFile* PakFile = nullptr;
+		FPakInputPair InputPair;
+		int64 OriginalFileSize = 0;
+		FOutputPakFileEntry* UExpFileInPair = nullptr;
+		bool bIsUAssetUExpPairUAsset = false;
+		bool bIsUAssetUExpPairUExp = false;
+		bool bIsMappedBulk = false;
+		bool bIsLastFileInPak = false;
+		bool bSomeCompressionSucceeded = false;
+		bool bCopiedToPak = false;
 
-	if (CmdLineParameters.bPatchCompatibilityMode421)
-	{
-		// for old versions, put in some known names that we may have used
-		Info.GetCompressionMethodIndex(NAME_None);
-		Info.GetCompressionMethodIndex(NAME_Zlib);
-		Info.GetCompressionMethodIndex(NAME_Gzip);
-		Info.GetCompressionMethodIndex(TEXT("Bogus"));
-		Info.GetCompressionMethodIndex(TEXT("Oodle"));
+		FPakEntryPair Entry;
+		FCompressedFileBuffer CompressedFileBuffer;
+		FName CompressionMethod = NAME_None;
+		int64 RealFileSize = 0;
 
-	}
+		FGraphEventRef CompressionBarrier;
+		FGraphEventRef CompressionTask;
+	};
 
-	if (InKeyChain.MasterEncryptionKey)
+	struct FOutputPakFile
 	{
-		UE_LOG(LogPakFile, Display, TEXT("Using encryption key '%s' [%s]"), *InKeyChain.MasterEncryptionKey->Name, *InKeyChain.MasterEncryptionKey->Guid.ToString());
-	}
+		FString Filename;
+		FKeyChain KeyChain;
+		TUniquePtr<FArchive> PakFileHandle;
+		TUniquePtr<FArchive> PakFileRegionsHandle;
+		TArray<FOutputPakFileEntry> Entries;
+		FPakInfo Info;
+		FString MountPoint;
+		TArray<FPakEntryPair> Index;
+		// Some platforms provide patch download size reduction by diffing the patch files.  However, they often operate on specific block
+		// sizes when dealing with new data within the file.  Pad files out to the given alignment to work with these systems more nicely.
+		// We also want to combine smaller files into the same padding size block so we don't waste as much space. i.e. grouping 64 1k files together
+		// rather than padding each out to 64k.
+		uint64 ContiguousTotalSizeSmallerThanBlockSize = 0;
+		uint64 ContiguousFilesSmallerThanBlockSize = 0;
+		TArray<FFileRegion> AllFileRegions;
+		// Stats
+		uint64 TotalUncompressedSize = 0;
+		uint64 TotalCompressedSize = 0;
+		TArray<int32> Compressor_Stat_Count;
+		TArray<uint64> Compressor_Stat_RawBytes;
+		TArray<uint64> Compressor_Stat_CompBytes;
+		uint64 TotalRequestedEncryptedFiles = 0;
+		uint64 TotalEncryptedFiles = 0;
+		uint64 TotalEncryptedDataSize = 0;
+		FEventRef AllEntriesRetiredEvent;
+		int32 RetiredEntriesCount = 0;
+	};
 
-	TArray<FPakEntryPair> Index;
-	FString MountPoint = GetCommonRootPath(FilesToAdd);
-	uint8* ReadBuffer = NULL;
-	int64 BufferSize = 0;
+	void Compress(FOutputPakFileEntry* Entry);
+	void Retire(FOutputPakFileEntry* Entry);
+	void CompressionThreadFunc();
+	void WriterThreadFunc();
 
-	uint8* PaddingBuffer = nullptr;
-	int64 PaddingBufferSize = 0;
-	if (CmdLineParameters.PatchFilePadAlign > 0 || CmdLineParameters.AlignForMemoryMapping)
-	{
-		PaddingBufferSize = FMath::Max(CmdLineParameters.PatchFilePadAlign, CmdLineParameters.AlignForMemoryMapping);
-		PaddingBuffer = (uint8*)FMemory::Malloc(PaddingBufferSize);
-		FMemory::Memset(PaddingBuffer, 0, PaddingBufferSize);
-	}
+	FPakCommandLineParameters CmdLineParameters;
+	TSet<FString> NoPluginCompressionFileNames;
+	TSet<FString> NoPluginCompressionExtensions;
+	TArray64<uint8> PaddingBuffer;
+	TArray<TUniquePtr<FOutputPakFile>> OutputPakFiles;
+	TArray<FName> CompressionFormatsAndNone;
+	TAtomic<int64> TotalFilesWithPoorForcedCompression{ 0 };
+	TAtomic<int64> TotalExtraMemoryForPoorForcedCompression{ 0 };
+	TSpscQueue<FOutputPakFileEntry*> CompressionQueue;
+	TSpscQueue<FOutputPakFileEntry*> WriteQueue;
+	TFuture<void> CompressionThread;
+	TFuture<void> WriterThread;
+	FEventRef CompressionQueueEntryAddedEvent;
+	FEventRef WriteQueueEntryAddedEvent;
+	FEventRef EntryRetiredEvent;
+	TAtomic<int64> ScheduledFileSize{ 0 };
+	TAtomic<bool> bFlushed{ false };
+	uint64 TotalEntriesCount = 0;
+	TAtomic<uint64> RetiredEntriesCount{ 0 };
+};
 
-	// Some platforms provide patch download size reduction by diffing the patch files.  However, they often operate on specific block
-	// sizes when dealing with new data within the file.  Pad files out to the given alignment to work with these systems more nicely.
-	// We also want to combine smaller files into the same padding size block so we don't waste as much space. i.e. grouping 64 1k files together
-	// rather than padding each out to 64k.
-	const uint32 RequiredPatchPadding = CmdLineParameters.PatchFilePadAlign;
+bool FPakWriterContext::Initialize(const FPakCommandLineParameters& InCmdLineParameters)
+{
+#if USE_DDC_FOR_COMPRESSED_FILES
+	GetDerivedDataCacheRef();
+#endif
 
-	uint64 ContiguousTotalSizeSmallerThanBlockSize = 0;
-	uint64 ContiguousFilesSmallerThanBlockSize = 0;
+	CmdLineParameters = InCmdLineParameters;
 
-	uint64 TotalUncompressedSize = 0;
-	uint64 TotalCompressedSize = 0;
+	int64 PaddingBufferSize = 64 * 1024;
+	PaddingBufferSize = FMath::Max(PaddingBufferSize, CmdLineParameters.AlignForMemoryMapping);
+	PaddingBufferSize = FMath::Max(PaddingBufferSize, CmdLineParameters.PatchFilePadAlign);
+	PaddingBuffer.SetNumZeroed(PaddingBufferSize);
 
-	uint64 TotalRequestedEncryptedFiles = 0;
-	uint64 TotalEncryptedFiles = 0;
-	uint64 TotalEncryptedDataSize = 0;
-	
 	// track compression stats per format :
-	TArray<FName> CompressionFormatsAndNone = CmdLineParameters.CompressionFormats;
+	CompressionFormatsAndNone = CmdLineParameters.CompressionFormats;
 	CompressionFormatsAndNone.AddUnique(NAME_None);
-	int32 CompressionFormatsAndNone_Num = CompressionFormatsAndNone.Num();
-	TArray<int32> Compressor_Stat_Count; Compressor_Stat_Count.SetNumZeroed(CompressionFormatsAndNone_Num);
-	TArray<uint64> Compressor_Stat_RawBytes; Compressor_Stat_RawBytes.SetNumZeroed(CompressionFormatsAndNone_Num);
-	TArray<uint64> Compressor_Stat_CompBytes; Compressor_Stat_CompBytes.SetNumZeroed(CompressionFormatsAndNone_Num);
 
 	{
 		// log the methods and indexes
@@ -1890,9 +1929,6 @@ bool CreatePakFile(const TCHAR* Filename, TArray<FPakInputPair>& FilesToAdd, con
 		UE_LOG(LogPakFile, Display, TEXT("%s"), *FormatLogLine);
 	}
 
-	TSet<FString> NoPluginCompressionFileNames;
-	TSet<FString> NoPluginCompressionExtensions;
-	
 	// Oodle is built into the Engine now and can be used to decode startup phase files (ini,res,uplugin)
 	//	 which used to be forced to Zlib
 	//	 set bDoUseOodleDespiteNoPluginCompression = true to let Oodle compress those files
@@ -1922,290 +1958,326 @@ bool CreatePakFile(const TCHAR* Filename, TArray<FPakInputPair>& FilesToAdd, con
 		}
 	}
 
-	struct FAsyncCompressor
+	CompressionThread = Async(EAsyncExecution::Thread, [this]()
+		{
+			CompressionThreadFunc();
+		});
+
+	WriterThread = Async(EAsyncExecution::Thread, [this]()
+		{
+			WriterThreadFunc();
+		});
+
+	return true;
+}
+
+bool FPakWriterContext::AddPakFile(const TCHAR* Filename, const TArray<FPakInputPair>& FilesToAdd, const FKeyChain& InKeyChain)
+{
+	TRACE_CPUPROFILER_EVENT_SCOPE(AddPakFile);
+	check(!bFlushed);
+
+	TUniquePtr<FOutputPakFile> OutputPakFile = MakeUnique<FOutputPakFile>();
+
+	OutputPakFile->Filename = Filename;
+	OutputPakFile->KeyChain = InKeyChain;
+	
 	{
-		
-		// input
-		const FPakInputPair* FileToAdd; 
-		FPakEntryPair Entry;
-		const TArray<FName>* CompressionFormats;
-		const TSet<FString>* NoPluginCompressionExtensions;
-		const TSet<FString>* NoPluginCompressionFileNames;
+		TRACE_CPUPROFILER_EVENT_SCOPE(CreateFileHandle);
+		OutputPakFile->PakFileHandle.Reset(CreatePakWriter(Filename, InKeyChain, CmdLineParameters.bSign));
+	}
+	if (!OutputPakFile->PakFileHandle)
+	{
+		UE_LOG(LogPakFile, Error, TEXT("Unable to create pak file \"%s\"."), Filename);
+		return false;
+	}
 
-		// output
-		FCompressedFileBuffer CompressedFileBuffer;
-		FName CompressionMethod; 
-		int32 CompressionBlockSize;
-		int64 RealFileSize;
-		int64 OriginalFileSize;
-		bool bForceCompress = false;
-		volatile bool bIsComplete;
-		
-		void Init(	FPakInputPair* InFileToAdd,
-					const FPakCommandLineParameters& InParams, 
-					const TSet<FString>* InNoPluginCompressionExtensions, 
-					const TSet<FString>* InNoPluginCompressionFileNames)
+	if (CmdLineParameters.bFileRegions)
+	{
+		FString RegionsFilename = FString(Filename) + FFileRegion::RegionsFileExtension;
 		{
-			CompressionMethod = NAME_None;
-			if (InFileToAdd->bIsDeleteRecord)
-			{
-				// don't need to do anything for delete records
-				bIsComplete = true;
-				return;
-			}
-			bForceCompress = InParams.bForceCompress;
-			FileToAdd = InFileToAdd;
-			NoPluginCompressionExtensions = InNoPluginCompressionExtensions;
-			NoPluginCompressionFileNames = InNoPluginCompressionFileNames;
-			CompressionFormats = &InParams.CompressionFormats;
-			CompressionBlockSize = InParams.CompressionBlockSize;
-			bIsComplete = false;
+			TRACE_CPUPROFILER_EVENT_SCOPE(CreateRegionsFileHandle);
+			OutputPakFile->PakFileRegionsHandle.Reset(IFileManager::Get().CreateFileWriter(*RegionsFilename));
 		}
-		
-		bool IsComplete() const
+		if (!OutputPakFile->PakFileRegionsHandle)
 		{
-			return bIsComplete;
+			UE_LOG(LogPakFile, Error, TEXT("Unable to create pak regions file \"%s\"."), *RegionsFilename);
+			return false;
 		}
-		void Complete()
+	}
+
+	if (InKeyChain.MasterEncryptionKey)
+	{
+		UE_LOG(LogPakFile, Display, TEXT("Using encryption key '%s' [%s]"), *InKeyChain.MasterEncryptionKey->Name, *InKeyChain.MasterEncryptionKey->Guid.ToString());
+	}
+
+	OutputPakFile->Info.bEncryptedIndex = (InKeyChain.MasterEncryptionKey && CmdLineParameters.EncryptIndex);
+	OutputPakFile->Info.EncryptionKeyGuid = InKeyChain.MasterEncryptionKey ? InKeyChain.MasterEncryptionKey->Guid : FGuid();
+
+	if (CmdLineParameters.bPatchCompatibilityMode421)
+	{
+		// for old versions, put in some known names that we may have used
+		OutputPakFile->Info.GetCompressionMethodIndex(NAME_None);
+		OutputPakFile->Info.GetCompressionMethodIndex(NAME_Zlib);
+		OutputPakFile->Info.GetCompressionMethodIndex(NAME_Gzip);
+		OutputPakFile->Info.GetCompressionMethodIndex(TEXT("Bogus"));
+		OutputPakFile->Info.GetCompressionMethodIndex(TEXT("Oodle"));
+	}
+
+	OutputPakFile->MountPoint = GetCommonRootPath(FilesToAdd);
+
+	OutputPakFile->Entries.SetNum(FilesToAdd.Num());
+	for (int32 FileIndex = 0; FileIndex < FilesToAdd.Num(); FileIndex++)
+	{
+		FOutputPakFileEntry& OutputEntry = OutputPakFile->Entries[FileIndex];
+		OutputEntry.PakFile = OutputPakFile.Get();
+		OutputEntry.InputPair = FilesToAdd[FileIndex];
+
+		OutputEntry.bIsMappedBulk = FilesToAdd[FileIndex].Source.EndsWith(TEXT(".m.ubulk"));
+		OutputEntry.bIsLastFileInPak = FileIndex + 1 == FilesToAdd.Num();
+		if (FileIndex > 0)
 		{
-			bIsComplete = true;
+			if (FPaths::GetBaseFilename(FilesToAdd[FileIndex - 1].Dest, false) == FPaths::GetBaseFilename(FilesToAdd[FileIndex].Dest, false) &&
+				FPaths::GetExtension(FilesToAdd[FileIndex - 1].Dest, true) == TEXT(".uasset") &&
+				FPaths::GetExtension(FilesToAdd[FileIndex].Dest, true) == TEXT(".uexp"))
+			{
+				OutputEntry.bIsUAssetUExpPairUExp = true;
+			}
 		}
-
-		void Compress()
+		if (!OutputEntry.bIsUAssetUExpPairUExp && FileIndex + 1 < FilesToAdd.Num())
 		{
-			if (bIsComplete == true)
+			if (FPaths::GetBaseFilename(FilesToAdd[FileIndex].Dest, false) == FPaths::GetBaseFilename(FilesToAdd[FileIndex + 1].Dest, false) &&
+				FPaths::GetExtension(FilesToAdd[FileIndex].Dest, true) == TEXT(".uasset") &&
+				FPaths::GetExtension(FilesToAdd[FileIndex + 1].Dest, true) == TEXT(".uexp"))
 			{
-				return;
+				OutputEntry.bIsUAssetUExpPairUAsset = true;
+				OutputEntry.UExpFileInPair = &OutputPakFile->Entries[FileIndex + 1];
 			}
+		}
+		OutputEntry.CompressionBarrier = FGraphEvent::CreateGraphEvent();
+		OutputEntry.CompressionTask = FFunctionGraphTask::CreateAndDispatchWhenReady([this, &OutputEntry]()
+			{
+				Compress(&OutputEntry);
+			}, TStatId(), OutputEntry.CompressionBarrier, ENamedThreads::AnyHiPriThreadHiPriTask);
+		CompressionQueue.Enqueue(&OutputEntry);
+		WriteQueue.Enqueue(&OutputEntry);
+	}
 
-			//check if this file requested to be compression
-			OriginalFileSize = IFileManager::Get().FileSize(*FileToAdd->Source);
-			RealFileSize = OriginalFileSize + Entry.Info.GetSerializedSize(FPakInfo::PakFile_Version_Latest);
+	CompressionQueueEntryAddedEvent->Trigger();
+	WriteQueueEntryAddedEvent->Trigger();
 
-			// don't try to compress tiny files
-			// even if they do compress, it is a bad use of decoder time
-			if (OriginalFileSize < 1024 || (FileToAdd->bNeedsCompression == false))
-			{
-				// done, don't need to do anything else
-				Complete();
-				return;
-			}
+	OutputPakFile->Compressor_Stat_Count.SetNumZeroed(CompressionFormatsAndNone.Num());
+	OutputPakFile->Compressor_Stat_RawBytes.SetNumZeroed(CompressionFormatsAndNone.Num());
+	OutputPakFile->Compressor_Stat_CompBytes.SetNumZeroed(CompressionFormatsAndNone.Num());
 
-			bool bSomeCompressionSucceeded = false;
+	TotalEntriesCount += OutputPakFile->Entries.Num();
+	if (OutputPakFile->Entries.IsEmpty())
+	{
+		OutputPakFile->AllEntriesRetiredEvent->Trigger();
+	}
+	OutputPakFiles.Add(MoveTemp(OutputPakFile));
 
-			// if first method (oodle) doesn't compress, don't try other methods (zlib)
-			// if Oodle refused to compress it was not an error, it was a choice because Oodle didn't think
-			//  compressing that file was worth the decode time
-			//	we do NOT want Zlib to then get enabled for that file!
-			for (int32 MethodIndex = 0; MethodIndex < 1; MethodIndex++)
-			{
-				CompressionMethod = (*CompressionFormats)[MethodIndex];
+	return true;
+}
 
-				// because compression is a plugin, certain files need to be loadable out of pak files before plugins are loadable
-				// (like .uplugin files). for these, we enforce a non-plugin compression - zlib
-				//	note that those file types are also excluded from iostore, so still go through this pak system
-				bool bForceCompressionFormat = false;
-				if (NoPluginCompressionExtensions->Find(FPaths::GetExtension(FileToAdd->Source)) != nullptr ||
-					NoPluginCompressionFileNames->Find(FPaths::GetCleanFilename(FileToAdd->Source)) != nullptr)
-				{
-					CompressionMethod = NAME_Zlib;
-					bForceCompressionFormat = true;
-				}
+void FPakWriterContext::Compress(FOutputPakFileEntry* Entry)
+{
+	TRACE_CPUPROFILER_EVENT_SCOPE(Compress);
 
-				// attempt to compress the data
-				if (CompressedFileBuffer.CompressFileToWorkingBuffer(*FileToAdd, CompressionMethod, CompressionBlockSize))
-				{
-					// for modern compressors we don't want any funny heuristics turning compression on/off
-					// let the compressor decide; it has an introspective measure of whether compression is worth doing or not
-					bool bNotEnoughCompression = CompressedFileBuffer.TotalCompressedSize >= OriginalFileSize;
+	// don't try to compress tiny files
+	// even if they do compress, it is a bad use of decoder time
+	if (Entry->OriginalFileSize < 1024 || !Entry->InputPair.bNeedsCompression)
+	{
+		Entry->CompressedFileBuffer.ReadUncompressedFileToWorkingBuffer(Entry->InputPair);
+		Entry->bSomeCompressionSucceeded = true; // Only used for logging purposes
+		return;
+	}
 
-					if ( CompressionMethod == NAME_Zlib )
-					{
-						// for forced-Zlib files still use the old heuristic :
+	// if first method (oodle) doesn't compress, don't try other methods (zlib)
+	// if Oodle refused to compress it was not an error, it was a choice because Oodle didn't think
+	//  compressing that file was worth the decode time
+	//	we do NOT want Zlib to then get enabled for that file!
+	for (int32 MethodIndex = 0; MethodIndex < 1; MethodIndex++)
+	{
+		Entry->CompressionMethod = CmdLineParameters.CompressionFormats[MethodIndex];
 
-						// Zlib must save at least 1K regardless of percentage (for small files)
-						bNotEnoughCompression = (OriginalFileSize - CompressedFileBuffer.TotalCompressedSize) < 1024;
-						if ( ! bNotEnoughCompression )
-						{
-							// Check the compression ratio, if it's too low just store uncompressed. Also take into account read size
-							// if we still save 64KB it's probably worthwhile compressing, as that saves a file read operation in the runtime.
-							// TODO: drive this threshold from the command line
-							float PercentLess = ((float)CompressedFileBuffer.TotalCompressedSize / (OriginalFileSize / 100.f));
-							bNotEnoughCompression = (PercentLess > 90.f) && ((OriginalFileSize - CompressedFileBuffer.TotalCompressedSize) < 65536);
-						}
-					}				
-				
-					const bool bIsLastCompressionFormat = MethodIndex == CompressionFormats->Num() - 1;
-					if (bNotEnoughCompression && (!bForceCompress || !bIsLastCompressionFormat))
-					{
-						// compression did not succeed, we can try the next format, so do nothing here
-					}
-					else
-					{
-						Entry.Info.CompressionBlocks.AddUninitialized(CompressedFileBuffer.CompressedBlocks.Num());
-						RealFileSize = CompressedFileBuffer.TotalCompressedSize + Entry.Info.GetSerializedSize(FPakInfo::PakFile_Version_Latest);
-						Entry.Info.CompressionBlocks.Reset();
+		// because compression is a plugin, certain files need to be loadable out of pak files before plugins are loadable
+		// (like .uplugin files). for these, we enforce a non-plugin compression - zlib
+		//	note that those file types are also excluded from iostore, so still go through this pak system
+		bool bForceCompressionFormat = false;
+		if (NoPluginCompressionExtensions.Find(FPaths::GetExtension(Entry->InputPair.Source)) != nullptr ||
+			NoPluginCompressionFileNames.Find(FPaths::GetCleanFilename(Entry->InputPair.Source)) != nullptr)
+		{
+			Entry->CompressionMethod = NAME_Zlib;
+			bForceCompressionFormat = true;
+		}
 
-						// at this point, we have successfully compressed the file, no need to continue
-						bSomeCompressionSucceeded = true;
+		// attempt to compress the data
+		if (Entry->CompressedFileBuffer.CompressFileToWorkingBuffer(Entry->InputPair, Entry->CompressionMethod, CmdLineParameters.CompressionBlockSize))
+		{
+			// for modern compressors we don't want any funny heuristics turning compression on/off
+			// let the compressor decide; it has an introspective measure of whether compression is worth doing or not
+			bool bNotEnoughCompression = Entry->CompressedFileBuffer.TotalCompressedSize >= Entry->OriginalFileSize;
 
-						if (bNotEnoughCompression)
-						{
-							// None of the compression formats were good enough, but we were under instructions to use some form of compression
-							// This was likely to aid with runtime error checking on mobile devices. Record how many cases of this we encountered and
-							// How much the size difference was
-							GTotalFilesWithPoorForcedCompression++;
-							GTotalExtraMemoryForPoorForcedCompression += FMath::Max(CompressedFileBuffer.TotalCompressedSize - OriginalFileSize, (int64)0);
-						}
-					}
-				}
+			if (Entry->CompressionMethod == NAME_Zlib)
+			{
+				// for forced-Zlib files still use the old heuristic :
 
-				// if we successfully compressed it, or we only wanted a single format, then we are done!
-				if (bSomeCompressionSucceeded || bForceCompressionFormat)
+				// Zlib must save at least 1K regardless of percentage (for small files)
+				bNotEnoughCompression = (Entry->OriginalFileSize - Entry->CompressedFileBuffer.TotalCompressedSize) < 1024;
+				if (!bNotEnoughCompression)
 				{
-					break;
+					// Check the compression ratio, if it's too low just store uncompressed. Also take into account read size
+					// if we still save 64KB it's probably worthwhile compressing, as that saves a file read operation in the runtime.
+					// TODO: drive this threshold from the command line
+					float PercentLess = ((float)Entry->CompressedFileBuffer.TotalCompressedSize / (Entry->OriginalFileSize / 100.f));
+					bNotEnoughCompression = (PercentLess > 90.f) && ((Entry->OriginalFileSize - Entry->CompressedFileBuffer.TotalCompressedSize) < 65536);
 				}
 			}
 
-			// If no compression was able to make it small enough, or compress at all, don't compress it
-			if (!bSomeCompressionSucceeded)
+			const bool bIsLastCompressionFormat = MethodIndex == CmdLineParameters.CompressionFormats.Num() - 1;
+			if (bNotEnoughCompression && (!CmdLineParameters.bForceCompress || !bIsLastCompressionFormat))
 			{
-				UE_LOG(LogPakFile, Log, TEXT("File \"%s\" did not get small enough from compression, or compression failed."), *FileToAdd->Source);
-				CompressionMethod = NAME_None;
+				// compression did not succeed, we can try the next format, so do nothing here
 			}
-			Complete();
-		}
-
-		void CleanUp()
-		{
-			CompressedFileBuffer.Empty();
-		}
-	};
-
-
+			else
+			{
+				Entry->Entry.Info.CompressionBlocks.AddUninitialized(Entry->CompressedFileBuffer.CompressedBlocks.Num());
+				Entry->RealFileSize = Entry->CompressedFileBuffer.TotalCompressedSize + Entry->Entry.Info.GetSerializedSize(FPakInfo::PakFile_Version_Latest);
+				Entry->Entry.Info.CompressionBlocks.Reset();
 
-	FThreadSafeCounter CompletionCounter;
+				// at this point, we have successfully compressed the file, no need to continue
+				Entry->bSomeCompressionSucceeded = true;
 
-	TArray<FAsyncCompressor> AsyncCompressors;
-	AsyncCompressors.AddZeroed(FilesToAdd.Num());
-	
-	class FRunCompressionTask : public FNonAbandonableTask
-	{
-		FAsyncCompressor* AsyncCompressor;
-	public:
-		FRunCompressionTask(FAsyncCompressor* InAsyncCompressor)
-		{
-			AsyncCompressor = InAsyncCompressor;
+				if (bNotEnoughCompression)
+				{
+					// None of the compression formats were good enough, but we were under instructions to use some form of compression
+					// This was likely to aid with runtime error checking on mobile devices. Record how many cases of this we encountered and
+					// How much the size difference was
+					TotalFilesWithPoorForcedCompression++;
+					TotalExtraMemoryForPoorForcedCompression += FMath::Max(Entry->CompressedFileBuffer.TotalCompressedSize - Entry->OriginalFileSize, (int64)0);
+				}
+			}
 		}
 
-		void DoWork()
+		// if we successfully compressed it, or we only wanted a single format, then we are done!
+		if (Entry->bSomeCompressionSucceeded || bForceCompressionFormat)
 		{
-			AsyncCompressor->Compress();
+			break;
 		}
+	}
 
-		FORCEINLINE TStatId GetStatId() const { RETURN_QUICK_DECLARE_CYCLE_STAT(ExampleAsyncTask, STATGROUP_ThreadPoolAsyncTasks); }
-	};
-
-	const bool bRunAsync = CmdLineParameters.bAsyncCompression;
-
-#if USE_DDC_FOR_COMPRESSED_FILES
-	GetDerivedDataCacheRef();
-#endif
-	
-	TArray<FFileRegion> AllFileRegions;
-	
-	// limit memory use :
-	// Launch async tasks ahead of our current retiring index
-	// enough to keep the cores busy
-	// but we don't want every file in the whole Pak
-	// because the compressed buffers are held after the task is done
+	// If no compression was able to make it small enough, or compress at all, don't compress it
+	if (!Entry->bSomeCompressionSucceeded)
+	{
+		Entry->CompressionMethod = NAME_None;
+		Entry->CompressedFileBuffer.ReadUncompressedFileToWorkingBuffer(Entry->InputPair);
+	}
+}
 
-	int NextLaunchIndex = 0;
-	int NumToLaunchAhead = 128;   // <- make NumToLaunchAhead lower to reduce memory use more
-	// TODO: NumToLaunchAhead should ideally scale up with 2*cores or so but also be limited by memory size.
+void FPakWriterContext::Retire(FOutputPakFileEntry* Entry)
+{
+	++RetiredEntriesCount;
+	Entry->CompressedFileBuffer.Empty();
+	ScheduledFileSize -= Entry->OriginalFileSize;
+	EntryRetiredEvent->Trigger();
+	++Entry->PakFile->RetiredEntriesCount;
+	if (Entry->PakFile->RetiredEntriesCount == Entry->PakFile->Entries.Num())
+	{
+		Entry->PakFile->AllEntriesRetiredEvent->Trigger();
+	}
+}
 
-	for (int32 FileIndex = 0; FileIndex < FilesToAdd.Num(); FileIndex++)
+void FPakWriterContext::CompressionThreadFunc()
+{
+	const int64 MaxScheduledFileSize = 2ll << 30;
+	for (;;)
 	{
-		// FileIndex is what we're currently retiring
-		// everything in [FileIndex,NextLaunchIndex) may be pending
-		// launch new jobs up to NumToLaunchAhead ahead of FileIndex
-		int32 LaunchFileIndexEnd = FMath::Min( FileIndex + NumToLaunchAhead, FilesToAdd.Num() );
-		for (int32 LaunchFileIndex = NextLaunchIndex; LaunchFileIndex < LaunchFileIndexEnd; LaunchFileIndex++)
+		TOptional<FOutputPakFileEntry*> FromQueue = CompressionQueue.Dequeue();
+		if (FromQueue.IsSet())
 		{
-			AsyncCompressors[LaunchFileIndex].Init(&FilesToAdd[LaunchFileIndex], CmdLineParameters, &NoPluginCompressionExtensions, &NoPluginCompressionFileNames);
+			FOutputPakFileEntry* Entry = FromQueue.GetValue();
+			if (!Entry->OriginalFileSize)
+			{
+				Entry->OriginalFileSize = IFileManager::Get().FileSize(*Entry->InputPair.Source);
+			}
+			Entry->RealFileSize = Entry->OriginalFileSize + Entry->Entry.Info.GetSerializedSize(FPakInfo::PakFile_Version_Latest);
+			if (Entry->UExpFileInPair)
+			{
+				// Need the size of the uexp as well before we can proceed with this entry
+				Entry->UExpFileInPair->OriginalFileSize = IFileManager::Get().FileSize(*Entry->UExpFileInPair->InputPair.Source);
+			}
 
-			if (bRunAsync)
+			int64 LocalScheduledFileSize = ScheduledFileSize;
+			if (LocalScheduledFileSize + Entry->OriginalFileSize > MaxScheduledFileSize)
 			{
-				if (FilesToAdd[LaunchFileIndex].bNeedsCompression)
-				{
-					(new FAutoDeleteAsyncTask<FRunCompressionTask>(&AsyncCompressors[LaunchFileIndex]))->StartBackgroundTask();
-				}
-				else
+				TRACE_CPUPROFILER_EVENT_SCOPE(WaitForMemory);
+				while (LocalScheduledFileSize > 0 && LocalScheduledFileSize + Entry->OriginalFileSize > MaxScheduledFileSize)
 				{
-					// call compress function inline 
-					// it won't do anything except for initialize some internal variables used in the non compressed path
-					// we don't want to pass these to a different thread as they may cause congestion with legitimate tasks
-					AsyncCompressors[LaunchFileIndex].Compress();
+					EntryRetiredEvent->Wait();
+					LocalScheduledFileSize = ScheduledFileSize;
 				}
 			}
+			ScheduledFileSize += Entry->OriginalFileSize;
+			Entry->CompressionBarrier->DispatchSubsequents();
 		}
-		NextLaunchIndex = LaunchFileIndexEnd;
-
-		bool bDeleted = FilesToAdd[FileIndex].bIsDeleteRecord;
-		bool bIsUAssetUExpPairUAsset = false;
-		bool bIsUAssetUExpPairUExp = false;
-		bool bIsMappedBulk = FilesToAdd[FileIndex].Source.EndsWith(TEXT(".m.ubulk"));
-
-
-		if (FileIndex)
+		else
 		{
-			if (FPaths::GetBaseFilename(FilesToAdd[FileIndex - 1].Dest, false) == FPaths::GetBaseFilename(FilesToAdd[FileIndex].Dest, false) &&
-				FPaths::GetExtension(FilesToAdd[FileIndex - 1].Dest, true) == TEXT(".uasset") &&
-				FPaths::GetExtension(FilesToAdd[FileIndex].Dest, true) == TEXT(".uexp")
-				)
+			if (bFlushed)
 			{
-				bIsUAssetUExpPairUExp = true;
+				return;
 			}
+			TRACE_CPUPROFILER_EVENT_SCOPE(WaitingForWork);
+			CompressionQueueEntryAddedEvent->Wait();
 		}
-		if (!bIsUAssetUExpPairUExp && FileIndex + 1 < FilesToAdd.Num())
+	}
+}
+
+void FPakWriterContext::WriterThreadFunc()
+{
+	const uint32 RequiredPatchPadding = CmdLineParameters.PatchFilePadAlign;
+
+	for (;;)
+	{
+		TOptional<FOutputPakFileEntry*> FromQueue = WriteQueue.Dequeue();
+		if (!FromQueue.IsSet())
 		{
-			if (FPaths::GetBaseFilename(FilesToAdd[FileIndex].Dest, false) == FPaths::GetBaseFilename(FilesToAdd[FileIndex + 1].Dest, false) &&
-				FPaths::GetExtension(FilesToAdd[FileIndex].Dest, true) == TEXT(".uasset") &&
-				FPaths::GetExtension(FilesToAdd[FileIndex + 1].Dest, true) == TEXT(".uexp")
-				)
+			if (bFlushed)
 			{
-				bIsUAssetUExpPairUAsset = true;
+				return;
 			}
+			TRACE_CPUPROFILER_EVENT_SCOPE(WaitingForWork);
+			WriteQueueEntryAddedEvent->Wait();
+			continue;
 		}
 
+		TRACE_CPUPROFILER_EVENT_SCOPE(ProcessFile);
+
+		FOutputPakFileEntry* OutputEntry = FromQueue.GetValue();
+		bool bDeleted = OutputEntry->InputPair.bIsDeleteRecord;
+
 		//  Remember the offset but don't serialize it with the entry header.
+		FArchive* PakFileHandle = OutputEntry->PakFile->PakFileHandle.Get();
 		int64 NewEntryOffset = PakFileHandle->Tell();
-		FPakEntryPair &NewEntry = AsyncCompressors[FileIndex].Entry;
-		
-		if (bRunAsync)
-		{
-			while (AsyncCompressors[FileIndex].IsComplete() != true)
-			{
-				FPlatformProcess::Sleep(0.001f);
-			}
-		}
-		else
+		FPakEntryPair& NewEntry = OutputEntry->Entry;
+
+		if (!OutputEntry->CompressionTask->IsComplete())
 		{
-			AsyncCompressors[FileIndex].Compress();
+			TRACE_CPUPROFILER_EVENT_SCOPE(WaitForCompression);
+			OutputEntry->CompressionTask->Wait();
 		}
-
-		const FName& CompressionMethod = AsyncCompressors[FileIndex].CompressionMethod;
-		FCompressedFileBuffer& CompressedFileBuffer = AsyncCompressors[FileIndex].CompressedFileBuffer;
+		
+		const FName& CompressionMethod = OutputEntry->CompressionMethod;
+		FCompressedFileBuffer& CompressedFileBuffer = OutputEntry->CompressedFileBuffer;
 
 		if (!bDeleted)
 		{
-			const int64 &RealFileSize = AsyncCompressors[FileIndex].RealFileSize;
-			const int64 &OriginalFileSize = AsyncCompressors[FileIndex].OriginalFileSize;
-			const int32 CompressionBlockSize = AsyncCompressors[FileIndex].CompressionBlockSize;
-			NewEntry.Info.CompressionMethodIndex = Info.GetCompressionMethodIndex(CompressionMethod);
+			const int64& RealFileSize = OutputEntry->RealFileSize;
+			const int64& OriginalFileSize = OutputEntry->OriginalFileSize;
+			const int32 CompressionBlockSize = CmdLineParameters.CompressionBlockSize;
+			NewEntry.Info.CompressionMethodIndex = OutputEntry->PakFile->Info.GetCompressionMethodIndex(CompressionMethod);
 
 			// Account for file system block size, which is a boundary we want to avoid crossing.
-			if (!bIsUAssetUExpPairUExp && // don't split uexp / uasset pairs
+			if (!OutputEntry->bIsUAssetUExpPairUExp && // don't split uexp / uasset pairs
 				CmdLineParameters.FileSystemBlockSize > 0 &&
 				OriginalFileSize != INDEX_NONE &&
 				(CmdLineParameters.bAlignFilesLargerThanBlock || RealFileSize <= CmdLineParameters.FileSystemBlockSize) &&
@@ -2219,19 +2291,11 @@ bool CreatePakFile(const TCHAR* Filename, TArray<FPakInputPair>& FilesToAdd, con
 				check(PaddingRequired < RealFileSize);
 				if (PaddingRequired > 0)
 				{
-					// If we don't already have a padding buffer, create one
-					if (PaddingBuffer == nullptr)
-					{
-						PaddingBufferSize = 64 * 1024;
-						PaddingBuffer = (uint8*)FMemory::Malloc(PaddingBufferSize);
-						FMemory::Memset(PaddingBuffer, 0, PaddingBufferSize);
-					}
-
 					UE_LOG(LogPakFile, Verbose, TEXT("%14llu - %14llu : %14llu padding."), PakFileHandle->Tell(), PakFileHandle->Tell() + PaddingRequired, PaddingRequired);
 					while (PaddingRequired > 0)
 					{
-						int64 AmountToWrite = FMath::Min(PaddingRequired, PaddingBufferSize);
-						PakFileHandle->Serialize(PaddingBuffer, AmountToWrite);
+						int64 AmountToWrite = FMath::Min(PaddingRequired, PaddingBuffer.Num());
+						PakFileHandle->Serialize(PaddingBuffer.GetData(), AmountToWrite);
 						PaddingRequired -= AmountToWrite;
 					}
 
@@ -2240,7 +2304,7 @@ bool CreatePakFile(const TCHAR* Filename, TArray<FPakInputPair>& FilesToAdd, con
 			}
 
 			// Align bulk data
-			if (bIsMappedBulk && CmdLineParameters.AlignForMemoryMapping > 0 && OriginalFileSize != INDEX_NONE && !bDeleted)
+			if (OutputEntry->bIsMappedBulk && CmdLineParameters.AlignForMemoryMapping > 0 && OriginalFileSize != INDEX_NONE && !bDeleted)
 			{
 				if (!IsAligned(NewEntryOffset + NewEntry.Info.GetSerializedSize(FPakInfo::PakFile_Version_Latest), CmdLineParameters.AlignForMemoryMapping))
 				{
@@ -2249,11 +2313,11 @@ bool CreatePakFile(const TCHAR* Filename, TArray<FPakInputPair>& FilesToAdd, con
 					int64 PaddingRequired = NewEntryOffset - OldOffset;
 
 					check(PaddingRequired > 0);
-					check(PaddingBuffer && PaddingBufferSize >= PaddingRequired);
+					check(PaddingBuffer.Num() >= PaddingRequired);
 
 					{
 						UE_LOG(LogPakFile, Verbose, TEXT("%14llu - %14llu : %14llu bulk padding."), PakFileHandle->Tell(), PakFileHandle->Tell() + PaddingRequired, PaddingRequired);
-						PakFileHandle->Serialize(PaddingBuffer, PaddingRequired);
+						PakFileHandle->Serialize(PaddingBuffer.GetData(), PaddingRequired);
 						check(PakFileHandle->Tell() == NewEntryOffset);
 					}
 				}
@@ -2262,45 +2326,42 @@ bool CreatePakFile(const TCHAR* Filename, TArray<FPakInputPair>& FilesToAdd, con
 
 		TArray<FFileRegion> CurrentFileRegions;
 
-		bool bCopiedToPak;
 		int64 SizeToWrite = 0;
 		uint8* DataToWrite = nullptr;
 		if (bDeleted)
 		{
-			PrepareDeleteRecordForPak(MountPoint, FilesToAdd[FileIndex], NewEntry);
-			bCopiedToPak = false;
+			PrepareDeleteRecordForPak(OutputEntry->PakFile->MountPoint, OutputEntry->InputPair, NewEntry);
+			OutputEntry->bCopiedToPak = false;
 
 			// Directly add the new entry to the index, no more work to do
-			Index.Add(NewEntry);
+			OutputEntry->PakFile->Index.Add(NewEntry);
 		}
-		else if (FilesToAdd[FileIndex].bNeedsCompression && CompressionMethod != NAME_None)
+		else if (OutputEntry->InputPair.bNeedsCompression && CompressionMethod != NAME_None)
 		{
-			bCopiedToPak = PrepareCopyCompressedFileToPak(MountPoint, Info, FilesToAdd[FileIndex], CompressedFileBuffer, NewEntry, DataToWrite, SizeToWrite, InKeyChain);
-			DataToWrite = CompressedFileBuffer.CompressedBuffer.Get();
+			OutputEntry->bCopiedToPak = PrepareCopyCompressedFileToPak(OutputEntry->PakFile->MountPoint, OutputEntry->PakFile->Info, OutputEntry->InputPair, CompressedFileBuffer, NewEntry, DataToWrite, SizeToWrite, OutputEntry->PakFile->KeyChain);
 		}
 		else
 		{
-			bCopiedToPak = PrepareCopyFileToPak(MountPoint, FilesToAdd[FileIndex], ReadBuffer, BufferSize, NewEntry, DataToWrite, SizeToWrite, InKeyChain, CmdLineParameters.bFileRegions ? &CurrentFileRegions : nullptr);
-			DataToWrite = ReadBuffer;
+			OutputEntry->bCopiedToPak = PrepareCopyFileToPak(OutputEntry->PakFile->MountPoint, OutputEntry->InputPair, CompressedFileBuffer, NewEntry, DataToWrite, SizeToWrite, OutputEntry->PakFile->KeyChain, CmdLineParameters.bFileRegions ? &CurrentFileRegions : nullptr);
 		}
 
 		int64 TotalSizeToWrite = SizeToWrite + NewEntry.Info.GetSerializedSize(FPakInfo::PakFile_Version_Latest);
-		if (bCopiedToPak)
+		if (OutputEntry->bCopiedToPak)
 		{
 			if (RequiredPatchPadding > 0 &&
-				!(bIsMappedBulk && CmdLineParameters.AlignForMemoryMapping > 0) // don't wreck the bulk padding with patch padding
+				!(OutputEntry->bIsMappedBulk && CmdLineParameters.AlignForMemoryMapping > 0) // don't wreck the bulk padding with patch padding
 				)
 			{
 				//if the next file is going to cross a patch-block boundary then pad out the current set of files with 0's
 				//and align the next file up.
 				bool bCrossesBoundary = (NewEntryOffset / RequiredPatchPadding) != ((NewEntryOffset + TotalSizeToWrite - 1) / RequiredPatchPadding);
 				bool bPatchPadded = false;
-				if (!bIsUAssetUExpPairUExp) // never patch-pad the uexp of a uasset/uexp pair
+				if (!OutputEntry->bIsUAssetUExpPairUExp) // never patch-pad the uexp of a uasset/uexp pair
 				{
 					bool bPairProbablyCrossesBoundary = false; // we don't consider compression because we have not compressed the uexp yet.
-					if (bIsUAssetUExpPairUAsset)
+					if (OutputEntry->bIsUAssetUExpPairUAsset)
 					{
-						int64 UExpFileSize = IFileManager::Get().FileSize(*FilesToAdd[FileIndex + 1].Source) / 2; // assume 50% compression
+						int64 UExpFileSize = OutputEntry->UExpFileInPair->OriginalFileSize / 2; // assume 50% compression
 						bPairProbablyCrossesBoundary = (NewEntryOffset / RequiredPatchPadding) != ((NewEntryOffset + TotalSizeToWrite + UExpFileSize - 1) / RequiredPatchPadding);
 					}
 					if (TotalSizeToWrite >= RequiredPatchPadding || // if it exactly the padding size and by luck does not cross a boundary, we still consider it "over" because it can't be packed with anything else
@@ -2313,11 +2374,11 @@ bool CreatePakFile(const TCHAR* Filename, TArray<FPakInputPair>& FilesToAdd, con
 						if (PaddingSize)
 						{
 							UE_LOG(LogPakFile, Verbose, TEXT("%14llu - %14llu : %14llu patch padding."), PakFileHandle->Tell(), PakFileHandle->Tell() + PaddingSize, PaddingSize);
-							check(PaddingSize <= PaddingBufferSize);
+							check(PaddingSize <= PaddingBuffer.Num());
 
 							//have to pad manually with 0's.  File locations skipped by Seek and never written are uninitialized which would defeat the whole purpose
 							//of padding for certain platforms patch diffing systems.
-							PakFileHandle->Serialize(PaddingBuffer, PaddingSize);
+							PakFileHandle->Serialize(PaddingBuffer.GetData(), PaddingSize);
 						}
 						check(PakFileHandle->Tell() == NewEntryOffset);
 						bPatchPadded = true;
@@ -2328,212 +2389,267 @@ bool CreatePakFile(const TCHAR* Filename, TArray<FPakInputPair>& FilesToAdd, con
 				//if there were a large set of contiguous small files behind us then this will be the natural stopping point for a possible pathalogical patching case where growth in the small files causes a cascade 
 				//to dirty up all the blocks prior to this one.  If this could happen let's warn about it.
 				if (bPatchPadded ||
-					FileIndex + 1 == FilesToAdd.Num()) // also check the last file, this won't work perfectly if we don't end up adding the last file for some reason
+					OutputEntry->bIsLastFileInPak) // also check the last file, this won't work perfectly if we don't end up adding the last file for some reason
 				{
 					const uint64 ContiguousGroupedFilePatchWarningThreshhold = 50 * 1024 * 1024;
-					if (ContiguousTotalSizeSmallerThanBlockSize > ContiguousGroupedFilePatchWarningThreshhold)
+					if (OutputEntry->PakFile->ContiguousTotalSizeSmallerThanBlockSize > ContiguousGroupedFilePatchWarningThreshhold)
 					{
-						UE_LOG(LogPakFile, Display, TEXT("%i small files (%i) totaling %llu contiguous bytes found before first 'large' file.  Changes to any of these files could cause the whole group to be 'dirty' in a per-file binary diff based patching system."), ContiguousFilesSmallerThanBlockSize, RequiredPatchPadding, ContiguousTotalSizeSmallerThanBlockSize);
+						UE_LOG(LogPakFile, Display, TEXT("%i small files (%i) totaling %llu contiguous bytes found before first 'large' file.  Changes to any of these files could cause the whole group to be 'dirty' in a per-file binary diff based patching system."), OutputEntry->PakFile->ContiguousFilesSmallerThanBlockSize, RequiredPatchPadding, OutputEntry->PakFile->ContiguousTotalSizeSmallerThanBlockSize);
 					}
-					ContiguousTotalSizeSmallerThanBlockSize = 0;
-					ContiguousFilesSmallerThanBlockSize = 0;
+					OutputEntry->PakFile->ContiguousTotalSizeSmallerThanBlockSize = 0;
+					OutputEntry->PakFile->ContiguousFilesSmallerThanBlockSize = 0;
 				}
 				else
 				{
-					ContiguousTotalSizeSmallerThanBlockSize += TotalSizeToWrite;
-					ContiguousFilesSmallerThanBlockSize++;
+					OutputEntry->PakFile->ContiguousTotalSizeSmallerThanBlockSize += TotalSizeToWrite;
+					OutputEntry->PakFile->ContiguousFilesSmallerThanBlockSize++;
 				}
 			}
-			if (FilesToAdd[FileIndex].bNeedsCompression && CompressionMethod != NAME_None)
+			if (OutputEntry->InputPair.bNeedsCompression && CompressionMethod != NAME_None)
 			{
-				FinalizeCopyCompressedFileToPak(Info, CompressedFileBuffer, NewEntry);
+				FinalizeCopyCompressedFileToPak(OutputEntry->PakFile->Info, CompressedFileBuffer, NewEntry);
 			}
-			
+
 			{
 				// track per-compressor stats :
 				// note GetCompressionMethodIndex in the pak entry is the index in the Pak file list of compressors
 				//	not the same as the index in the command line list
 
 				int32 CompressionMethodIndex;
-				if ( CompressionFormatsAndNone.Find(CompressionMethod,CompressionMethodIndex) )
+				if (CompressionFormatsAndNone.Find(CompressionMethod, CompressionMethodIndex))
 				{
-					Compressor_Stat_Count[CompressionMethodIndex] += 1;
-					Compressor_Stat_RawBytes[ CompressionMethodIndex] += NewEntry.Info.UncompressedSize;
-					Compressor_Stat_CompBytes[CompressionMethodIndex] += NewEntry.Info.Size;
+					OutputEntry->PakFile->Compressor_Stat_Count[CompressionMethodIndex] += 1;
+					OutputEntry->PakFile->Compressor_Stat_RawBytes[CompressionMethodIndex] += NewEntry.Info.UncompressedSize;
+					OutputEntry->PakFile->Compressor_Stat_CompBytes[CompressionMethodIndex] += NewEntry.Info.Size;
 				}
 			}
 
 			// Write to file
-			int64 Offset = PakFileHandle->Tell();
-			NewEntry.Info.Serialize(*PakFileHandle, FPakInfo::PakFile_Version_Latest);
-			int64 PayloadOffset = PakFileHandle->Tell();
-			PakFileHandle->Serialize(DataToWrite, SizeToWrite);
-			int64 EndOffset = PakFileHandle->Tell();
-
-			if (CmdLineParameters.bFileRegions)
 			{
-				FFileRegion::AccumulateFileRegions(AllFileRegions, Offset, PayloadOffset, EndOffset, CurrentFileRegions);
-			}
+				TRACE_CPUPROFILER_EVENT_SCOPE(WriteToFile);
+				int64 Offset = PakFileHandle->Tell();
+				NewEntry.Info.Serialize(*PakFileHandle, FPakInfo::PakFile_Version_Latest);
+				int64 PayloadOffset = PakFileHandle->Tell();
+				PakFileHandle->Serialize(DataToWrite, SizeToWrite);
+				int64 EndOffset = PakFileHandle->Tell();
 
-			UE_LOG(LogPakFile, Verbose, TEXT("%14llu [header] - %14llu - %14llu : %14llu header+file %s."), Offset, PayloadOffset, EndOffset, EndOffset - Offset, *NewEntry.Filename);
+				if (CmdLineParameters.bFileRegions)
+				{
+					FFileRegion::AccumulateFileRegions(OutputEntry->PakFile->AllFileRegions, Offset, PayloadOffset, EndOffset, CurrentFileRegions);
+				}
+
+				UE_LOG(LogPakFile, Verbose, TEXT("%14llu [header] - %14llu - %14llu : %14llu header+file %s."), Offset, PayloadOffset, EndOffset, EndOffset - Offset, *NewEntry.Filename);
+			}
 
 			// Update offset now and store it in the index (and only in index)
 			NewEntry.Info.Offset = NewEntryOffset;
-			Index.Add(NewEntry);
+			OutputEntry->PakFile->Index.Add(NewEntry);
 			const TCHAR* EncryptedString = TEXT("");
 
-			if (FilesToAdd[FileIndex].bNeedEncryption)
+			if (OutputEntry->InputPair.bNeedEncryption)
 			{
-				TotalRequestedEncryptedFiles++;
+				OutputEntry->PakFile->TotalRequestedEncryptedFiles++;
 
-				if (InKeyChain.MasterEncryptionKey)
+				if (OutputEntry->PakFile->KeyChain.MasterEncryptionKey)
 				{
-					TotalEncryptedFiles++;
-					TotalEncryptedDataSize += SizeToWrite;
+					OutputEntry->PakFile->TotalEncryptedFiles++;
+					OutputEntry->PakFile->TotalEncryptedDataSize += SizeToWrite;
 					EncryptedString = TEXT("encrypted ");
 				}
 			}
 
-			if (FilesToAdd[FileIndex].bNeedsCompression && CompressionMethod != NAME_None)
+			if (OutputEntry->InputPair.bNeedsCompression && CompressionMethod != NAME_None)
+			{
+				OutputEntry->PakFile->TotalCompressedSize += NewEntry.Info.Size;
+				OutputEntry->PakFile->TotalUncompressedSize += NewEntry.Info.UncompressedSize;
+			}
+		}
+		Retire(OutputEntry);
+	}
+}
+
+bool FPakWriterContext::Flush()
+{
+	TRACE_CPUPROFILER_EVENT_SCOPE(Flush);
+	check(!bFlushed);
+	bFlushed = true;
+
+	const uint32 RequiredPatchPadding = CmdLineParameters.PatchFilePadAlign;
+
+	for (TUniquePtr<FOutputPakFile>& OutputPakFile : OutputPakFiles)
+	{
+		{
+			TRACE_CPUPROFILER_EVENT_SCOPE(WaitForWritesToComplete);
+			for (;;)
+			{
+				bool bCompleted = OutputPakFile->AllEntriesRetiredEvent->Wait(2000);
+				if (bCompleted)
+				{
+					break;
+				}
+				UE_LOG(LogPakFile, Display, TEXT("Writing entries (%llu/%llu)..."), RetiredEntriesCount.Load(), TotalEntriesCount);
+			}
+		}
+
+		TRACE_CPUPROFILER_EVENT_SCOPE(FinalizePakFile);
+		UE_LOG(LogPakFile, Display, TEXT("Created pak file: %s"), *OutputPakFile->Filename);
+
+		for (FOutputPakFileEntry& OutputEntry : OutputPakFile->Entries)
+		{
+			if (!OutputEntry.bSomeCompressionSucceeded)
+			{
+				UE_LOG(LogPakFile, Verbose, TEXT("File \"%s\" did not get small enough from compression, or compression failed."), *OutputEntry.InputPair.Source);
+			}
+			if (OutputEntry.bCopiedToPak)
 			{
-				TotalCompressedSize += NewEntry.Info.Size;
-				TotalUncompressedSize += NewEntry.Info.UncompressedSize;
-				float PercentLess = ((float)NewEntry.Info.Size / (NewEntry.Info.UncompressedSize / 100.f));
-				if (FilesToAdd[FileIndex].SuggestedOrder < MAX_uint64)
+				const TCHAR* EncryptedString = TEXT("");
+
+				if (OutputEntry.InputPair.bNeedEncryption)
 				{
-					UE_LOG(LogPakFile, Log, TEXT("Added compressed %sfile \"%s\", %.2f%% of original size. Compressed with %s, Size %lld bytes, Original Size %lld bytes (order %llu)."), EncryptedString, *NewEntry.Filename, PercentLess, *CompressionMethod.ToString(), NewEntry.Info.Size, NewEntry.Info.UncompressedSize, FilesToAdd[FileIndex].SuggestedOrder);
+					if (OutputEntry.PakFile->KeyChain.MasterEncryptionKey)
+					{
+						EncryptedString = TEXT("encrypted ");
+					}
+				}
+
+				if (OutputEntry.InputPair.bNeedsCompression && OutputEntry.CompressionMethod != NAME_None)
+				{
+					float PercentLess = ((float)OutputEntry.Entry.Info.Size / (OutputEntry.Entry.Info.UncompressedSize / 100.f));
+					if (OutputEntry.InputPair.SuggestedOrder < MAX_uint64)
+					{
+						UE_LOG(LogPakFile, Verbose, TEXT("Added compressed %sfile \"%s\", %.2f%% of original size. Compressed with %s, Size %lld bytes, Original Size %lld bytes (order %llu)."), EncryptedString, *OutputEntry.Entry.Filename, PercentLess, *OutputEntry.CompressionMethod.ToString(), OutputEntry.Entry.Info.Size, OutputEntry.Entry.Info.UncompressedSize, OutputEntry.InputPair.SuggestedOrder);
+					}
+					else
+					{
+						UE_LOG(LogPakFile, Verbose, TEXT("Added compressed %sfile \"%s\", %.2f%% of original size. Compressed with %s, Size %lld bytes, Original Size %lld bytes (no order given)."), EncryptedString, *OutputEntry.Entry.Filename, PercentLess, *OutputEntry.CompressionMethod.ToString(), OutputEntry.Entry.Info.Size, OutputEntry.Entry.Info.UncompressedSize);
+					}
 				}
 				else
 				{
-					UE_LOG(LogPakFile, Log, TEXT("Added compressed %sfile \"%s\", %.2f%% of original size. Compressed with %s, Size %lld bytes, Original Size %lld bytes (no order given)."), EncryptedString, *NewEntry.Filename, PercentLess, *CompressionMethod.ToString(), NewEntry.Info.Size, NewEntry.Info.UncompressedSize);
+					if (OutputEntry.InputPair.SuggestedOrder < MAX_uint64)
+					{
+						UE_LOG(LogPakFile, Verbose, TEXT("Added %sfile \"%s\", %lld bytes (order %llu)."), EncryptedString, *OutputEntry.Entry.Filename, OutputEntry.Entry.Info.Size, OutputEntry.InputPair.SuggestedOrder);
+					}
+					else
+					{
+						UE_LOG(LogPakFile, Verbose, TEXT("Added %sfile \"%s\", %lld bytes (no order given)."), EncryptedString, *OutputEntry.Entry.Filename, OutputEntry.Entry.Info.Size);
+					}
 				}
 			}
 			else
 			{
-				if (FilesToAdd[FileIndex].SuggestedOrder < MAX_uint64)
+				if (OutputEntry.InputPair.bIsDeleteRecord)
 				{
-					UE_LOG(LogPakFile, Log, TEXT("Added %sfile \"%s\", %lld bytes (order %llu)."), EncryptedString, *NewEntry.Filename, NewEntry.Info.Size, FilesToAdd[FileIndex].SuggestedOrder);
+					UE_LOG(LogPakFile, Verbose, TEXT("Created delete record for file \"%s\"."), *OutputEntry.InputPair.Source);
 				}
 				else
 				{
-					UE_LOG(LogPakFile, Log, TEXT("Added %sfile \"%s\", %lld bytes (no order given)."), EncryptedString, *NewEntry.Filename, NewEntry.Info.Size);
+					UE_LOG(LogPakFile, Warning, TEXT("Missing file \"%s\" will not be added to PAK file."), *OutputEntry.InputPair.Source);
 				}
 			}
 		}
-		else
-		{
-			if (bDeleted)
-			{
-				UE_LOG(LogPakFile, Log, TEXT("Created delete record for file \"%s\"."), *FilesToAdd[FileIndex].Source);
-			}
-			else
-			{
-				UE_LOG(LogPakFile, Warning, TEXT("Missing file \"%s\" will not be added to PAK file."), *FilesToAdd[FileIndex].Source);
-			}
-		}
-		AsyncCompressors[FileIndex].CleanUp();
-	}
-
-	FMemory::Free(PaddingBuffer);
-	FMemory::Free(ReadBuffer);
-	ReadBuffer = NULL;
 
-	if (RequiredPatchPadding > 0)
-	{
-		for (const FPakEntryPair& Pair : Index)
+		if (RequiredPatchPadding > 0)
 		{
-			const FString& EntryFilename = Pair.Filename;
-			const FPakEntry& PakEntry = Pair.Info;
-			int64 EntrySize = PakEntry.GetSerializedSize(FPakInfo::PakFile_Version_Latest);
-			int64 TotalSizeToWrite = PakEntry.Size + EntrySize;
-			if (TotalSizeToWrite >= RequiredPatchPadding)
+			for (const FPakEntryPair& Pair : OutputPakFile->Index)
 			{
-				int64 RealStart = PakEntry.Offset;
-				if ((RealStart % RequiredPatchPadding) != 0 &&
-					!EntryFilename.EndsWith(TEXT("uexp")) && // these are export sections of larger files and may be packed with uasset/umap and so we don't need a warning here
-					!(EntryFilename.EndsWith(TEXT(".m.ubulk")) && CmdLineParameters.AlignForMemoryMapping > 0)) // Bulk padding unaligns patch padding and so we don't need a warning here
+				const FString& EntryFilename = Pair.Filename;
+				const FPakEntry& PakEntry = Pair.Info;
+				int64 EntrySize = PakEntry.GetSerializedSize(FPakInfo::PakFile_Version_Latest);
+				int64 TotalSizeToWrite = PakEntry.Size + EntrySize;
+				if (TotalSizeToWrite >= RequiredPatchPadding)
 				{
-					UE_LOG(LogPakFile, Warning, TEXT("File at offset %lld of size %lld not aligned to patch size %i"), RealStart, PakEntry.Size, RequiredPatchPadding);
+					int64 RealStart = PakEntry.Offset;
+					if ((RealStart % RequiredPatchPadding) != 0 &&
+						!EntryFilename.EndsWith(TEXT("uexp")) && // these are export sections of larger files and may be packed with uasset/umap and so we don't need a warning here
+						!(EntryFilename.EndsWith(TEXT(".m.ubulk")) && CmdLineParameters.AlignForMemoryMapping > 0)) // Bulk padding unaligns patch padding and so we don't need a warning here
+					{
+						UE_LOG(LogPakFile, Warning, TEXT("File at offset %lld of size %lld not aligned to patch size %i"), RealStart, PakEntry.Size, RequiredPatchPadding);
+					}
 				}
 			}
 		}
-	}
 
 
-	FPakFooterInfo Footer(Filename, MountPoint, Info, Index);
-	Footer.SetEncryptionInfo(InKeyChain, &TotalEncryptedDataSize);
-	Footer.SetFileRegionInfo(CmdLineParameters.bFileRegions, AllFileRegions);
-	WritePakFooter(*PakFileHandle, Footer);
+		FPakFooterInfo Footer(*OutputPakFile->Filename, OutputPakFile->MountPoint, OutputPakFile->Info, OutputPakFile->Index);
+		Footer.SetEncryptionInfo(OutputPakFile->KeyChain, &OutputPakFile->TotalEncryptedDataSize);
+		Footer.SetFileRegionInfo(CmdLineParameters.bFileRegions, OutputPakFile->AllFileRegions);
+		WritePakFooter(*OutputPakFile->PakFileHandle, Footer);
 
-	
-	// log per-compressor stats :
-	for (int32 MethodIndex = 0; MethodIndex < CompressionFormatsAndNone_Num; MethodIndex++)
-	{
-		FName CompressionMethod = CompressionFormatsAndNone[MethodIndex];
-		UE_LOG(LogPakFile, Display, TEXT("CompressionFormat %d [%s] : %d files, %lld -> %lld bytes"), MethodIndex, *(CompressionMethod.ToString()),
-			Compressor_Stat_Count[MethodIndex],
-			Compressor_Stat_RawBytes[MethodIndex],
-			Compressor_Stat_CompBytes[MethodIndex]
-			);
-	}
+		if (CmdLineParameters.bSign)
+		{
+			TArray<uint8> SignatureData;
+			SignatureData.Append(OutputPakFile->Info.IndexHash.Hash, UE_ARRAY_COUNT(FSHAHash::Hash));
+			((FSignedArchiveWriter*)OutputPakFile->PakFileHandle.Get())->SetSignatureData(SignatureData);
+		}
 
-	UE_LOG(LogPakFile, Display, TEXT("Added %d files, %lld bytes total, time %.2lfs."), Index.Num(), PakFileHandle->TotalSize(), FPlatformTime::Seconds() - StartTime);
-	UE_LOG(LogPakFile, Display, TEXT("PrimaryIndex size: %d bytes"), Footer.PrimaryIndexSize);
-	UE_LOG(LogPakFile, Display, TEXT("PathHashIndex size: %d bytes"), Footer.PathHashIndexSize);
-	UE_LOG(LogPakFile, Display, TEXT("FullDirectoryIndex size: %d bytes"), Footer.FullDirectoryIndexSize);
-	if (TotalUncompressedSize)
-	{
-		float PercentLess = ((float)TotalCompressedSize / (TotalUncompressedSize / 100.f));
-		UE_LOG(LogPakFile, Display, TEXT("Compression summary: %.2f%% of original size. Compressed Size %lld bytes, Original Size %lld bytes. "), PercentLess, TotalCompressedSize, TotalUncompressedSize);
+		uint64 TotalPakFileSize = OutputPakFile->PakFileHandle->TotalSize();
+		OutputPakFile->PakFileHandle->Close();
+		OutputPakFile->PakFileHandle.Reset();
 
-		if (GTotalFilesWithPoorForcedCompression > 0)
+		if (CmdLineParameters.bFileRegions)
 		{
-			UE_LOG(LogPakFile, Display, TEXT("Num files forcibly compressed due to -forcecompress option: %i, using %i bytes extra"), (int64)GTotalFilesWithPoorForcedCompression, (int64)GTotalExtraMemoryForPoorForcedCompression);
+			FFileRegion::SerializeFileRegions(*OutputPakFile->PakFileRegionsHandle.Get(), OutputPakFile->AllFileRegions);
+			OutputPakFile->PakFileRegionsHandle->Close();
+			OutputPakFile->PakFileRegionsHandle.Reset();
 		}
-	}
-
-	if (TotalEncryptedDataSize)
-	{
-		UE_LOG(LogPakFile, Display, TEXT("Encryption - ENABLED"));
-		UE_LOG(LogPakFile, Display, TEXT("  Files: %d"), TotalEncryptedFiles);
 
-		if (Info.bEncryptedIndex)
+		// log per-compressor stats :
+		for (int32 MethodIndex = 0; MethodIndex < CompressionFormatsAndNone.Num(); MethodIndex++)
 		{
-			UE_LOG(LogPakFile, Display, TEXT("  Index: Encrypted (%d bytes, %.2fMB)"), Info.IndexSize, (float)Info.IndexSize / 1024.0f / 1024.0f);
+			FName CompressionMethod = CompressionFormatsAndNone[MethodIndex];
+			UE_LOG(LogPakFile, Display, TEXT("CompressionFormat %d [%s] : %d files, %lld -> %lld bytes"), MethodIndex, *(CompressionMethod.ToString()),
+				OutputPakFile->Compressor_Stat_Count[MethodIndex],
+				OutputPakFile->Compressor_Stat_RawBytes[MethodIndex],
+				OutputPakFile->Compressor_Stat_CompBytes[MethodIndex]
+			);
 		}
-		else
+
+		UE_LOG(LogPakFile, Display, TEXT("Added %d files, %lld bytes total"), OutputPakFile->Index.Num(), TotalPakFileSize);
+		UE_LOG(LogPakFile, Display, TEXT("PrimaryIndex size: %d bytes"), Footer.PrimaryIndexSize);
+		UE_LOG(LogPakFile, Display, TEXT("PathHashIndex size: %d bytes"), Footer.PathHashIndexSize);
+		UE_LOG(LogPakFile, Display, TEXT("FullDirectoryIndex size: %d bytes"), Footer.FullDirectoryIndexSize);
+		if (OutputPakFile->TotalUncompressedSize)
 		{
-			UE_LOG(LogPakFile, Display, TEXT("  Index: Unencrypted"));
+			float PercentLess = ((float)OutputPakFile->TotalCompressedSize / (OutputPakFile->TotalUncompressedSize / 100.f));
+			UE_LOG(LogPakFile, Display, TEXT("Compression summary: %.2f%% of original size. Compressed Size %lld bytes, Original Size %lld bytes. "), PercentLess, OutputPakFile->TotalCompressedSize, OutputPakFile->TotalUncompressedSize);
 		}
 
+		if (OutputPakFile->TotalEncryptedDataSize)
+		{
+			UE_LOG(LogPakFile, Display, TEXT("Encryption - ENABLED"));
+			UE_LOG(LogPakFile, Display, TEXT("  Files: %d"), OutputPakFile->TotalEncryptedFiles);
 
-		UE_LOG(LogPakFile, Display, TEXT("  Total: %d bytes (%.2fMB)"), TotalEncryptedDataSize, (float)TotalEncryptedDataSize / 1024.0f / 1024.0f);
-	}
-	else
-	{
-		UE_LOG(LogPakFile, Display, TEXT("Encryption - DISABLED"));
-	}
+			if (OutputPakFile->Info.bEncryptedIndex)
+			{
+				UE_LOG(LogPakFile, Display, TEXT("  Index: Encrypted (%d bytes, %.2fMB)"), OutputPakFile->Info.IndexSize, (float)OutputPakFile->Info.IndexSize / 1024.0f / 1024.0f);
+			}
+			else
+			{
+				UE_LOG(LogPakFile, Display, TEXT("  Index: Unencrypted"));
+			}
 
-	if (TotalEncryptedFiles < TotalRequestedEncryptedFiles)
-	{
-		UE_LOG(LogPakFile, Display, TEXT("%d files requested encryption, but no AES key was supplied! Encryption was skipped for these files"), TotalRequestedEncryptedFiles);
-	}
 
-	if (CmdLineParameters.bSign)
-	{
-		TArray<uint8> SignatureData;
-		SignatureData.Append(Info.IndexHash.Hash, UE_ARRAY_COUNT(FSHAHash::Hash));
-		((FSignedArchiveWriter*)PakFileHandle.Get())->SetSignatureData(SignatureData);
+			UE_LOG(LogPakFile, Display, TEXT("  Total: %d bytes (%.2fMB)"), OutputPakFile->TotalEncryptedDataSize, (float)OutputPakFile->TotalEncryptedDataSize / 1024.0f / 1024.0f);
+		}
+		else
+		{
+			UE_LOG(LogPakFile, Display, TEXT("Encryption - DISABLED"));
+		}
+
+		if (OutputPakFile->TotalEncryptedFiles < OutputPakFile->TotalRequestedEncryptedFiles)
+		{
+			UE_LOG(LogPakFile, Display, TEXT("%d files requested encryption, but no AES key was supplied! Encryption was skipped for these files"), OutputPakFile->TotalRequestedEncryptedFiles);
+		}
+		UE_LOG(LogPakFile, Display, TEXT(""));
 	}
 
-	PakFileHandle->Close();
-	PakFileHandle.Reset();
+	WriterThread.Wait();
+	CompressionThread.Wait();
 
-	if (CmdLineParameters.bFileRegions)
+	if (TotalFilesWithPoorForcedCompression > 0)
 	{
-		FFileRegion::SerializeFileRegions(*PakFileRegionsHandle.Get(), AllFileRegions);
-		PakFileRegionsHandle->Close();
-		PakFileRegionsHandle.Reset();
+		UE_LOG(LogPakFile, Display, TEXT("Num files forcibly compressed due to -forcecompress option: %i, using %i bytes extra"), (int64)TotalFilesWithPoorForcedCompression, (int64)TotalExtraMemoryForPoorForcedCompression);
 	}
 
 #if DETAILED_UNREALPAK_TIMING
@@ -2547,6 +2663,25 @@ bool CreatePakFile(const TCHAR* Filename, TArray<FPakInputPair>& FilesToAdd, con
 	return true;
 }
 
+bool CreatePakFile(const TCHAR* Filename, const TArray<FPakInputPair>& FilesToAdd, const FPakCommandLineParameters& CmdLineParameters, const FKeyChain& InKeyChain)
+{
+	const double StartTime = FPlatformTime::Seconds();
+
+	// Create Pak
+	FPakWriterContext PakWriterContext;
+	if (!PakWriterContext.Initialize(CmdLineParameters))
+	{
+		return false;
+	}
+	
+	if (!PakWriterContext.AddPakFile(Filename, FilesToAdd, InKeyChain))
+	{
+		return false;
+	}
+
+	return PakWriterContext.Flush();
+}
+
 void WritePakFooter(FArchive& PakHandle, FPakFooterInfo& Footer)
 {
 	// Create the second copy of the FPakEntries stored in the Index at the end of the PakFile
@@ -5373,20 +5508,35 @@ bool ExecuteUnrealPak(const TCHAR* CmdLine)
 		return true;
 	}
 
-	if(NonOptionArguments.Num() > 0)
+	TArray<FString> CreatePakCommandsList;
 	{
+		FString CreateMultipleCommandListPath;
+		if (FParse::Value(CmdLine, TEXT("-CreateMultiple="), CreateMultipleCommandListPath))
+		{
+			if (!FFileHelper::LoadFileToStringArray(CreatePakCommandsList, *CreateMultipleCommandListPath))
+			{
+				UE_LOG(LogPakFile, Error, TEXT("Failed to read command list file '%s'."), *CreateMultipleCommandListPath);
+				return false;
+			}
+		}
+	}
+	if(NonOptionArguments.Num() > 0 || !CreatePakCommandsList.IsEmpty())
+	{
+		if (CreatePakCommandsList.IsEmpty())
+		{
+			CreatePakCommandsList.Add(CmdLine);
+		}
 
 		CheckAndReallocThreadPool();
 
+		FPakCommandLineParameters CommonCmdLineParameters;
+		ProcessCommonCommandLine(CmdLine, CommonCmdLineParameters);
 
-
-		// since this is for creation, we pass true to make it not look in LaunchDir
-		FString PakFilename = GetPakPath(*NonOptionArguments[0], true);
-
-		// List of all items to add to pak file
-		TArray<FPakInputPair> Entries;
-		FPakCommandLineParameters CmdLineParameters;
-		ProcessCommandLine(CmdLine, NonOptionArguments, Entries, CmdLineParameters);
+		FPakWriterContext PakWriterContext;
+		if (!PakWriterContext.Initialize(CommonCmdLineParameters))
+		{
+			return false;
+		}
 
 		FPakOrderMap OrderMap;
 		FString GameOpenOrderStr;
@@ -5420,114 +5570,147 @@ bool ExecuteUnrealPak(const TCHAR* CmdLine)
 			}
 		}
 
-		int32 LowestSourcePakVersion = 0;
-		TMap<FString, FFileInfo> SourceFileHashes;
+		// Check command line for the "patchcryptokeys" param, which will tell us where to look for the encryption keys that
+		// we need to access the patch reference data
+		FString PatchReferenceCryptoKeysFilename;
+		FKeyChain PatchKeyChain;
 
-		if ( CmdLineParameters.GeneratePatch )
+		if (FParse::Value(FCommandLine::Get(), TEXT("PatchCryptoKeys="), PatchReferenceCryptoKeysFilename))
 		{
-			FString OutputPath;
-			if (!FParse::Value(CmdLine, TEXT("TempFiles="), OutputPath))
-			{
-				OutputPath = FPaths::GetPath(PakFilename) / FString(TEXT("TempFiles"));
-			}
-
-			IFileManager::Get().DeleteDirectory(*OutputPath);
+			KeyChainUtilities::LoadKeyChainFromFile(PatchReferenceCryptoKeysFilename, PatchKeyChain);
+			KeyChainUtilities::ApplyEncryptionKeys(PatchKeyChain);
+		}
+		else
+		{
+			PatchKeyChain = KeyChain;
+		}
 
-			// Check command line for the "patchcryptokeys" param, which will tell us where to look for the encryption keys that
-			// we need to access the patch reference data
-			FString PatchReferenceCryptoKeysFilename;
-			FKeyChain PatchKeyChain;
+		KeyChainUtilities::ApplyEncryptionKeys(KeyChain);
 
-			if (FParse::Value(FCommandLine::Get(), TEXT("PatchCryptoKeys="), PatchReferenceCryptoKeysFilename))
+		TArray<FString> TempOutputDirectoriesToDelete;
+		for (const FString& CreatePakCommand : CreatePakCommandsList)
+		{
+			TArray<FString> NonOptionArgumentsForPakFile;
+			for (const TCHAR* CmdLineEnd = *CreatePakCommand; *CmdLineEnd != 0;)
 			{
-				KeyChainUtilities::LoadKeyChainFromFile(PatchReferenceCryptoKeysFilename, PatchKeyChain);
-				KeyChainUtilities::ApplyEncryptionKeys(PatchKeyChain);
+				FString Argument = FParse::Token(CmdLineEnd, false);
+				if (Argument.Len() > 0 && !Argument.StartsWith(TEXT("-")))
+				{
+					NonOptionArgumentsForPakFile.Add(Argument);
+				}
 			}
-			else
+
+			// since this is for creation, we pass true to make it not look in LaunchDir
+			FString PakFilename = GetPakPath(*NonOptionArgumentsForPakFile[0], true);
+			
+			// List of all items to add to pak file
+			TArray<FPakInputPair> Entries;			
+			FPakCommandLineParameters CmdLineParameters = CommonCmdLineParameters;
+			ProcessPakFileSpecificCommandLine(*CreatePakCommand, NonOptionArgumentsForPakFile, Entries, CmdLineParameters);
+			
+			FKeyChain KeyChainForPakFile = KeyChain;
+			FString EncryptionKeyOverrideGuidString;
+			FGuid EncryptionKeyOverrideGuid;
+			if (FParse::Value(*CreatePakCommand, TEXT("EncryptionKeyOverrideGuid="), EncryptionKeyOverrideGuidString))
 			{
-				PatchKeyChain = KeyChain;
+				FGuid::Parse(EncryptionKeyOverrideGuidString, EncryptionKeyOverrideGuid);
 			}
+			KeyChainForPakFile.MasterEncryptionKey = KeyChainForPakFile.EncryptionKeys.Find(EncryptionKeyOverrideGuid);
+		
+			int32 LowestSourcePakVersion = 0;
+			TMap<FString, FFileInfo> SourceFileHashes;
 
-			UE_LOG(LogPakFile, Display, TEXT("Generating patch from %s."), *CmdLineParameters.SourcePatchPakFilename, true );
-
-			TArray<FGuid> UsedEncryptionKeys;
-			if (!GenerateHashesFromPak(*CmdLineParameters.SourcePatchPakFilename, *PakFilename, SourceFileHashes, true, PatchKeyChain, /*Out*/LowestSourcePakVersion, &UsedEncryptionKeys))
+			if (CmdLineParameters.GeneratePatch)
 			{
-				UE_LOG(LogPakFile, Warning, TEXT("Unable to generate hashes from pak file %s"), *CmdLineParameters.SourcePatchPakFilename);
-				if (ExtractFilesFromPak(*CmdLineParameters.SourcePatchPakFilename, SourceFileHashes, *OutputPath, true, PatchKeyChain, nullptr, nullptr, nullptr, nullptr, &UsedEncryptionKeys) == false)
-				{
-					UE_LOG(LogPakFile, Warning, TEXT("Unable to extract files from source pak file for patch"));
-				}
-				else
+				FString TempOutputPath;
+				if (!FParse::Value(*CreatePakCommand, TEXT("TempFiles="), TempOutputPath))
 				{
-					CmdLineParameters.SourcePatchDiffDirectory = OutputPath;
+					TempOutputPath = FPaths::GetPath(PakFilename) / FString(TEXT("TempFiles"));
 				}
-			}
 
-			KeyChainUtilities::ApplyEncryptionKeys(KeyChain);
+				IFileManager::Get().DeleteDirectory(*TempOutputPath);
 
+				UE_LOG(LogPakFile, Display, TEXT("Generating patch from %s."), *CmdLineParameters.SourcePatchPakFilename, true);
 
-			if (UsedEncryptionKeys.Num() == 1)
-			{
-				if (UsedEncryptionKeys[0].IsValid())
+				TArray<FGuid> UsedEncryptionKeys;
+				if (!GenerateHashesFromPak(*CmdLineParameters.SourcePatchPakFilename, *PakFilename, SourceFileHashes, true, PatchKeyChain, /*Out*/LowestSourcePakVersion, &UsedEncryptionKeys))
+				{
+					UE_LOG(LogPakFile, Warning, TEXT("Unable to generate hashes from pak file %s"), *CmdLineParameters.SourcePatchPakFilename);
+					if (ExtractFilesFromPak(*CmdLineParameters.SourcePatchPakFilename, SourceFileHashes, *TempOutputPath, true, PatchKeyChain, nullptr, nullptr, nullptr, nullptr, &UsedEncryptionKeys) == false)
+					{
+						UE_LOG(LogPakFile, Warning, TEXT("Unable to extract files from source pak file for patch"));
+					}
+					else
+					{
+						TempOutputDirectoriesToDelete.Add(TempOutputPath);
+						CmdLineParameters.SourcePatchDiffDirectory = TempOutputPath;
+					}
+				}
+
+				if (UsedEncryptionKeys.Num() == 1)
 				{
-					UE_LOG(LogPakFile, Display, TEXT("Found encryption key %s in pak file %s will use to encrypt patch"), *((UsedEncryptionKeys[0]).ToString()), *CmdLineParameters.SourcePatchPakFilename);
-					KeyChain.MasterEncryptionKey = KeyChain.EncryptionKeys.Find(UsedEncryptionKeys[0]);
+					if (UsedEncryptionKeys[0].IsValid())
+					{
+						UE_LOG(LogPakFile, Display, TEXT("Found encryption key %s in pak file %s will use to encrypt patch"), *((UsedEncryptionKeys[0]).ToString()), *CmdLineParameters.SourcePatchPakFilename);
+						KeyChainForPakFile.MasterEncryptionKey = KeyChainForPakFile.EncryptionKeys.Find(UsedEncryptionKeys[0]);
+					}
 				}
 			}
-		}
 
+			// Start collecting files
+			TArray<FPakInputPair> FilesToAdd;
+			CollectFilesToAdd(FilesToAdd, Entries, OrderMap, CmdLineParameters);
 
-		// Start collecting files
-		TArray<FPakInputPair> FilesToAdd;
-		CollectFilesToAdd(FilesToAdd, Entries, OrderMap, CmdLineParameters);
+			if (FParse::Param(CmdLine, TEXT("writepakchunkorder")))
+			{
+				FString PakOrderFilename = FPaths::GetPath(PakFilename) + FPaths::GetBaseFilename(PakFilename) + FString(TEXT("-order")) + FString(TEXT(".txt"));
+				FArchive* PakOrderListArchive = IFileManager::Get().CreateFileWriter(*PakOrderFilename);
+				PakOrderListArchive->SetIsTextFormat(true);
 
-		if(FParse::Param(CmdLine, TEXT("writepakchunkorder")))
-		{
-			FString PakOrderFilename = FPaths::GetPath(PakFilename) + FPaths::GetBaseFilename(PakFilename) + FString(TEXT("-order")) + FString(TEXT(".txt"));
-			FArchive* PakOrderListArchive = IFileManager::Get().CreateFileWriter(*PakOrderFilename);
-			PakOrderListArchive->SetIsTextFormat(true);
+				for (int32 I = 0; I < FilesToAdd.Num(); ++I)
+				{
+					const FPakInputPair& Entry = FilesToAdd[I];
+					FString Line = FString::Printf(TEXT("%s %lld"), *Entry.Source, Entry.SuggestedOrder);
+					PakOrderListArchive->Logf(TEXT("%s"), *Line);
+				}
+				PakOrderListArchive->Close();
+				delete PakOrderListArchive;
+			}
 
-			for (int32 I = 0; I < FilesToAdd.Num(); ++I)
+			if (CmdLineParameters.GeneratePatch)
 			{
-				const FPakInputPair& Entry = FilesToAdd[I];
-				FString Line = FString::Printf(TEXT("%s %lld"), *Entry.Source, Entry.SuggestedOrder);
-				PakOrderListArchive->Logf(TEXT("%s"), *Line);
-			}
-			PakOrderListArchive->Close();
-			delete PakOrderListArchive;
-		}
+				// We need to get a list of files that were in the previous patch('s) Pak, but NOT in FilesToAdd
+				TArray<FPakInputPair> DeleteRecords = GetNewDeleteRecords(FilesToAdd, SourceFileHashes);
 
-		if ( CmdLineParameters.GeneratePatch )
-		{
-			// We need to get a list of files that were in the previous patch('s) Pak, but NOT in FilesToAdd
-			TArray<FPakInputPair> DeleteRecords = GetNewDeleteRecords(FilesToAdd, SourceFileHashes);
+				//if the patch is built using old source pak files, we need to handle the special case where a file has been moved between chunks but no delete record was created (this would cause a rogue delete record to be created in the latest pak), and also a case where the file was moved between chunks and back again without being changed (this would cause the file to not be included in this chunk because the file would be considered unchanged)
+				if (LowestSourcePakVersion < FPakInfo::PakFile_Version_DeleteRecords)
+				{
+					int32 CurrentPatchChunkIndex = GetPakChunkIndexFromFilename(PakFilename);
 
-			//if the patch is built using old source pak files, we need to handle the special case where a file has been moved between chunks but no delete record was created (this would cause a rogue delete record to be created in the latest pak), and also a case where the file was moved between chunks and back again without being changed (this would cause the file to not be included in this chunk because the file would be considered unchanged)
-			if (LowestSourcePakVersion < FPakInfo::PakFile_Version_DeleteRecords)
-			{
-				int32 CurrentPatchChunkIndex = GetPakChunkIndexFromFilename(PakFilename);
+					UE_LOG(LogPakFile, Display, TEXT("Some patch source paks were generated with an earlier version of UnrealPak that didn't support delete records. checking for historic assets that have moved between chunks to avoid creating invalid delete records"));
+					FString SourcePakFolder = FPaths::GetPath(CmdLineParameters.SourcePatchPakFilename);
 
-				UE_LOG(LogPakFile, Display, TEXT("Some patch source paks were generated with an earlier version of UnrealPak that didn't support delete records. checking for historic assets that have moved between chunks to avoid creating invalid delete records"));
-				FString SourcePakFolder = FPaths::GetPath(CmdLineParameters.SourcePatchPakFilename);
+					//remove invalid items from DeleteRecords and set 'bForceInclude' on some SourceFileHashes
+					ProcessLegacyFileMoves(DeleteRecords, SourceFileHashes, SourcePakFolder, FilesToAdd, CurrentPatchChunkIndex);
+				}
+				FilesToAdd.Append(DeleteRecords);
 
-				//remove invalid items from DeleteRecords and set 'bForceInclude' on some SourceFileHashes
-				ProcessLegacyFileMoves(DeleteRecords, SourceFileHashes, SourcePakFolder, FilesToAdd, CurrentPatchChunkIndex);
+				// if we are generating a patch here we remove files which are already shipped...
+				RemoveIdenticalFiles(FilesToAdd, CmdLineParameters.SourcePatchDiffDirectory, SourceFileHashes, CmdLineParameters.SeekOptParams, CmdLineParameters.ChangedFilesOutputFilename);
 			}
-			FilesToAdd.Append(DeleteRecords);
 
-			// if we are generating a patch here we remove files which are already shipped...
-			RemoveIdenticalFiles(FilesToAdd, CmdLineParameters.SourcePatchDiffDirectory, SourceFileHashes, CmdLineParameters.SeekOptParams, CmdLineParameters.ChangedFilesOutputFilename);
+			if (!PakWriterContext.AddPakFile(*PakFilename, FilesToAdd, KeyChainForPakFile))
+			{
+				return false;
+			}
 		}
 
+		bool bResult = PakWriterContext.Flush();
 
-		bool bResult = CreatePakFile(*PakFilename, FilesToAdd, CmdLineParameters, KeyChain);
-
-		if (CmdLineParameters.GeneratePatch)
+		for (const FString& TempOutputDirectory : TempOutputDirectoriesToDelete)
 		{
-			FString OutputPath = FPaths::GetPath(PakFilename) / FString(TEXT("TempFiles"));
 			// delete the temporary directory
-			IFileManager::Get().DeleteDirectory(*OutputPath, false, true);
+			IFileManager::Get().DeleteDirectory(*TempOutputDirectory, false, true);
 		}
 
 #if USE_DDC_FOR_COMPRESSED_FILES
diff --git a/Engine/Source/Programs/AutomationTool/Scripts/CopyBuildToStagingDirectory.Automation.cs b/Engine/Source/Programs/AutomationTool/Scripts/CopyBuildToStagingDirectory.Automation.cs
--- a/Engine/Source/Programs/AutomationTool/Scripts/CopyBuildToStagingDirectory.Automation.cs
+++ b/Engine/Source/Programs/AutomationTool/Scripts/CopyBuildToStagingDirectory.Automation.cs
@@ -183,19 +183,9 @@ namespace AutomationScripts
 			return Result;
 		}
 
-		static public string GetUnrealPakArguments(FileReference ProjectPath, Dictionary<string, string> UnrealPakResponseFile, FileReference OutputLocation, List<OrderFile> PakOrderFileLocations, string PlatformOptions, bool Compressed, EncryptionAndSigning.CryptoSettings CryptoSettings, FileReference CryptoKeysCacheFilename, String PatchSourceContentPath, string EncryptionKeyGuid, List<OrderFile> SecondaryPakOrderFileLocations, bool bUnattended)
+		private static string GetCommonUnrealPakArguments(List<OrderFile> PakOrderFileLocations, string AdditionalOptions, EncryptionAndSigning.CryptoSettings CryptoSettings, FileReference CryptoKeysCacheFilename, List<OrderFile> SecondaryPakOrderFileLocations, bool bUnattended)
 		{
-			StringBuilder CmdLine = new StringBuilder(MakePathSafeToUseWithCommandLine(ProjectPath.FullName));
-			CmdLine.AppendFormat(" {0}", MakePathSafeToUseWithCommandLine(OutputLocation.FullName));
-
-			// Force encryption of ALL files if we're using specific encryption key. This should be made an option per encryption key in the settings, but for our initial
-			// implementation we will just assume that we require maximum security for this data.
-			bool bForceEncryption = !string.IsNullOrEmpty(EncryptionKeyGuid);
-			string PakName = Path.GetFileNameWithoutExtension(OutputLocation.FullName);
-			string UnrealPakResponseFileName = CombinePaths(CmdEnv.LogFolder, "PakList_" + PakName + ".txt");
-			WritePakResponseFile(UnrealPakResponseFileName, UnrealPakResponseFile, Compressed, CryptoSettings, bForceEncryption);
-			CmdLine.AppendFormat(" -create={0}", CommandUtils.MakePathSafeToUseWithCommandLine(UnrealPakResponseFileName));
-
+			StringBuilder CmdLine = new StringBuilder();
 			if (CryptoKeysCacheFilename != null)
 			{
 				CmdLine.AppendFormat(" -cryptokeys={0}", CommandUtils.MakePathSafeToUseWithCommandLine(CryptoKeysCacheFilename.FullName));
@@ -208,20 +198,13 @@ namespace AutomationScripts
 			{
 				CmdLine.AppendFormat(" -secondaryOrder={0}", CommandUtils.MakePathSafeToUseWithCommandLine(string.Join(",", SecondaryPakOrderFileLocations.Select(u => u.File.FullName).ToArray())));
 			}
-			if (!String.IsNullOrEmpty(PatchSourceContentPath))
-			{
-				CmdLine.AppendFormat(" -generatepatch={0} -tempfiles={1}", CommandUtils.MakePathSafeToUseWithCommandLine(PatchSourceContentPath), CommandUtils.MakePathSafeToUseWithCommandLine(CommandUtils.CombinePaths(CommandUtils.CmdEnv.LocalRoot, "TempFiles" + Path.GetFileNameWithoutExtension(OutputLocation.FullName))));
-			}
+
 			if (CryptoSettings != null && CryptoSettings.bDataCryptoRequired)
 			{
 				if (CryptoSettings.bEnablePakIndexEncryption)
 				{
 					CmdLine.AppendFormat(" -encryptindex");
 				}
-				if (!string.IsNullOrEmpty(EncryptionKeyGuid))
-				{
-					CmdLine.AppendFormat(" -EncryptionKeyOverrideGuid={0}", EncryptionKeyGuid);
-				}
 				if (CryptoSettings.bDataCryptoRequired && CryptoSettings.bEnablePakSigning && CryptoSettings.SigningKey.IsValid())
 				{
 					CmdLine.AppendFormat(" -sign");
@@ -234,8 +217,45 @@ namespace AutomationScripts
 				CmdLine.AppendFormat(" -unattended");
 			}
 
-			CmdLine.Append(PlatformOptions);
+			CmdLine.Append(AdditionalOptions);
+
+			return CmdLine.ToString();
+		}
+
+		static private string GetPakFileSpecificUnrealPakArguments(Dictionary<string, string> UnrealPakResponseFile, FileReference OutputLocation, string AdditionalOptions, bool Compressed, EncryptionAndSigning.CryptoSettings CryptoSettings, String PatchSourceContentPath, string EncryptionKeyGuid)
+		{
+			StringBuilder CmdLine = new StringBuilder(MakePathSafeToUseWithCommandLine(OutputLocation.FullName));
+
+			// Force encryption of ALL files if we're using specific encryption key. This should be made an option per encryption key in the settings, but for our initial
+			// implementation we will just assume that we require maximum security for this data.
+			bool bForceEncryption = !string.IsNullOrEmpty(EncryptionKeyGuid);
+			string PakName = Path.GetFileNameWithoutExtension(OutputLocation.FullName);
+			string UnrealPakResponseFileName = CombinePaths(CmdEnv.LogFolder, "PakList_" + PakName + ".txt");
+			WritePakResponseFile(UnrealPakResponseFileName, UnrealPakResponseFile, Compressed, CryptoSettings, bForceEncryption);
+			CmdLine.AppendFormat(" -create={0}", CommandUtils.MakePathSafeToUseWithCommandLine(UnrealPakResponseFileName));
+
+			if (!String.IsNullOrEmpty(PatchSourceContentPath))
+			{
+				CmdLine.AppendFormat(" -generatepatch={0} -tempfiles={1}", CommandUtils.MakePathSafeToUseWithCommandLine(PatchSourceContentPath), CommandUtils.MakePathSafeToUseWithCommandLine(CommandUtils.CombinePaths(CommandUtils.CmdEnv.LocalRoot, "TempFiles" + Path.GetFileNameWithoutExtension(OutputLocation.FullName))));
+			}
+			if (CryptoSettings != null && CryptoSettings.bDataCryptoRequired)
+			{
+				if (!string.IsNullOrEmpty(EncryptionKeyGuid))
+				{
+					CmdLine.AppendFormat(" -EncryptionKeyOverrideGuid={0}", EncryptionKeyGuid);
+				}
+			}
+
+			return CmdLine.ToString();
+		}
 
+		static public string GetUnrealPakArguments(FileReference ProjectPath, Dictionary<string, string> UnrealPakResponseFile, FileReference OutputLocation, List<OrderFile> PakOrderFileLocations, string AdditionalOptions, bool Compressed, EncryptionAndSigning.CryptoSettings CryptoSettings, FileReference CryptoKeysCacheFilename, String PatchSourceContentPath, string EncryptionKeyGuid, List<OrderFile> SecondaryPakOrderFileLocations, bool bUnattended)
+		{
+			StringBuilder CmdLine = new StringBuilder(MakePathSafeToUseWithCommandLine(ProjectPath.FullName));
+			CmdLine.Append(" ");
+			CmdLine.Append(GetPakFileSpecificUnrealPakArguments(UnrealPakResponseFile, OutputLocation, "", Compressed, CryptoSettings, PatchSourceContentPath, EncryptionKeyGuid));
+			CmdLine.Append(" ");
+			CmdLine.Append(GetCommonUnrealPakArguments(PakOrderFileLocations, AdditionalOptions, CryptoSettings, CryptoKeysCacheFilename, SecondaryPakOrderFileLocations, bUnattended));
 			return CmdLine.ToString();
 		}
 
@@ -2555,7 +2575,7 @@ namespace AutomationScripts
 
 			Func<string, string> GetPostFix = PakFilename => GetPakFilePostFix(bShouldGeneratePatch, Params.HasIterateSharedCookedBuild, PlatformGameConfig, PakFilename);
 
-			List<string> Commands = new List<string>();
+			List<string> PakCommands = new List<string>();
 			List<string> LogNames = new List<string>();
 			List<string> IoStoreCommands = new List<string>();
 
@@ -2606,21 +2626,21 @@ namespace AutomationScripts
 				}
 			}
 
+			ConfigHierarchy PlatformEngineConfig;
+			bool bHasPlatformEngineConfig = false;
+			bHasPlatformEngineConfig = Params.EngineConfigs.TryGetValue(SC.StageTargetPlatform.PlatformType, out PlatformEngineConfig);
 			string BulkOption = "";
+			if (bHasPlatformEngineConfig)
 			{
-				ConfigHierarchy PlatformEngineConfig;
-				if (Params.EngineConfigs.TryGetValue(SC.StageTargetPlatform.PlatformType, out PlatformEngineConfig))
+				bool bMasterEnable = false;
+				PlatformEngineConfig.GetBool("MemoryMappedFiles", "MasterEnable", out bMasterEnable);
+				if (bMasterEnable)
 				{
-					bool bMasterEnable = false;
-					PlatformEngineConfig.GetBool("MemoryMappedFiles", "MasterEnable", out bMasterEnable);
-					if (bMasterEnable)
+					int Value = 0;
+					PlatformEngineConfig.GetInt32("MemoryMappedFiles", "Alignment", out Value);
+					if (Value > 0)
 					{
-						int Value = 0;
-						PlatformEngineConfig.GetInt32("MemoryMappedFiles", "Alignment", out Value);
-						if (Value > 0)
-						{
-							BulkOption = String.Format(" -AlignForMemoryMapping={0}", Value);
-						}
+						BulkOption = String.Format(" -AlignForMemoryMapping={0}", Value);
 					}
 				}
 			}
@@ -2773,6 +2793,42 @@ namespace AutomationScripts
 				LogInformation("    {0} Priority {1}", File.File.ToString(), File.Priority);
 			}
 
+			List<OrderFile> PrimaryOrderFiles = OrderFiles.FindAll(x => (x.OrderType != OrderFile.OrderFileType.Cooker));
+			List<OrderFile> SecondaryOrderFiles = null;
+
+			// Add a secondary order if there is one specified
+			if (bUseSecondaryOrder && OrderFiles.Count >= 1)
+			{
+				SecondaryOrderFiles = OrderFiles.FindAll(x => x.OrderType == OrderFile.OrderFileType.Cooker);
+			}
+
+			string CommonAdditionalArgs = SC.StageTargetPlatform.GetPlatformPakCommandLine(Params, SC) + BulkOption + PatchOptions + CompressionFormats + " " + Params.AdditionalPakOptions;
+
+			bool bPakFallbackOrderForNonUassetFiles = false;
+			PlatformGameConfig.GetBool("/Script/UnrealEd.ProjectPackagingSettings", "bPakFallbackOrderForNonUassetFiles", out bPakFallbackOrderForNonUassetFiles);
+			if (bPakFallbackOrderForNonUassetFiles)
+			{
+				CommonAdditionalArgs += " -fallbackOrderForNonUassetFiles";
+			}
+
+			if (bHasPlatformEngineConfig)
+			{
+				// if the runtime will want to reduce memory usage, we have to disable the pak index freezing
+				bool bUnloadPakEntries = false;
+				bool bShrinkPakEntries = false;
+				PlatformEngineConfig.GetBool("Pak", "UnloadPakEntryFilenamesIfPossible", out bUnloadPakEntries);
+				PlatformEngineConfig.GetBool("Pak", "ShrinkPakEntriesMemoryUsage", out bShrinkPakEntries);
+				if (bUnloadPakEntries || bShrinkPakEntries)
+				{
+					CommonAdditionalArgs += " -allowForIndexUnload";
+				}
+			}
+
+			// pass the targetplatform so the index may be able to be frozen
+			CommonAdditionalArgs += " -platform=" + ConfigHierarchy.GetIniPlatformName(SC.StageTargetPlatform.IniPlatformType);
+
+			CommonAdditionalArgs += " " + AdditionalCompressionOptionsOnCommandLine;
+
 			List<Tuple<FileReference, StagedFileReference, string>> Outputs = new List<Tuple<FileReference, StagedFileReference, string>>();
 			foreach (CreatePakParams PakParams in PakParamsList)
 			{
@@ -2919,45 +2975,6 @@ namespace AutomationScripts
 					}
 					if (!bCopiedExistingPak)
 					{
-						List<OrderFile> PrimaryOrderFiles = OrderFiles.FindAll(x => (x.OrderType != OrderFile.OrderFileType.Cooker));
-						List<OrderFile> SecondaryOrderFiles = null;
-
-						// Add a secondary order if there is one specified
-						if (bUseSecondaryOrder && OrderFiles.Count >= 1)
-						{
-							SecondaryOrderFiles = OrderFiles.FindAll(x => x.OrderType == OrderFile.OrderFileType.Cooker);
-						}
-
-						string AdditionalArgs = PatchOptions;
-						if (bShouldGeneratePatch && ShouldSkipGeneratingPatch(PlatformGameConfig, PakParams.PakName))
-						{
-							AdditionalArgs = " -SkipPatch";
-						}
-
-						bool bPakFallbackOrderForNonUassetFiles = false;
-						PlatformGameConfig.GetBool("/Script/UnrealEd.ProjectPackagingSettings", "bPakFallbackOrderForNonUassetFiles", out bPakFallbackOrderForNonUassetFiles);
-						if (bPakFallbackOrderForNonUassetFiles)
-						{
-							AdditionalArgs += " -fallbackOrderForNonUassetFiles";
-						}
-
-						ConfigHierarchy PlatformEngineConfig;
-						if (Params.EngineConfigs.TryGetValue(SC.StageTargetPlatform.PlatformType, out PlatformEngineConfig))
-						{
-							// if the runtime will want to reduce memory usage, we have to disable the pak index freezing
-							bool bUnloadPakEntries = false;
-							bool bShrinkPakEntries = false;
-							PlatformEngineConfig.GetBool("Pak", "UnloadPakEntryFilenamesIfPossible", out bUnloadPakEntries);
-							PlatformEngineConfig.GetBool("Pak", "ShrinkPakEntriesMemoryUsage", out bShrinkPakEntries);
-							if (bUnloadPakEntries || bShrinkPakEntries)
-							{
-								AdditionalArgs += " -allowForIndexUnload";
-							}
-						}
-
-						// pass the targetplatform so the index may be able to be frozen
-						AdditionalArgs += " -platform=" + ConfigHierarchy.GetIniPlatformName(SC.StageTargetPlatform.IniPlatformType);
-
 						Dictionary<string, string> UnrealPakResponseFile = PakParams.UnrealPakResponseFile;
 						if (bShouldCreateIoStoreContainerFiles)
 						{
@@ -3017,20 +3034,20 @@ namespace AutomationScripts
 								Params.HasDLCName));
 						}
 
+						string PakFileSpecificAdditionalArgs = "";
+						if (bShouldGeneratePatch && ShouldSkipGeneratingPatch(PlatformGameConfig, PakParams.PakName))
+						{
+							PakFileSpecificAdditionalArgs = " -SkipPatch";
+						}
 						string PakEncryptionKeyGuid = Params.SkipEncryption ? "" : PakParams.EncryptionKeyGuid;
-						Commands.Add(GetUnrealPakArguments(
-							Params.RawProjectPath,
+						PakCommands.Add(GetPakFileSpecificUnrealPakArguments(
 							UnrealPakResponseFile,
 							OutputLocation,
-							PrimaryOrderFiles,
-							SC.StageTargetPlatform.GetPlatformPakCommandLine(Params, SC) + AdditionalArgs + BulkOption + CompressionFormats + " " + Params.AdditionalPakOptions,
+							PakFileSpecificAdditionalArgs,
 							PakParams.bCompressed,
 							Params.SkipEncryption ? null : CryptoSettings,
-							Params.SkipEncryption ? null : CryptoKeysCacheFilename,
 							PatchSourceContentPath,
-							Params.SkipEncryption ? "" : PakParams.EncryptionKeyGuid,
-							SecondaryOrderFiles,
-							Params.Unattended));
+							Params.SkipEncryption ? "" : PakParams.EncryptionKeyGuid));
 
 						LogNames.Add(OutputLocation.GetFileNameWithoutExtension());
 					}
@@ -3040,9 +3057,21 @@ namespace AutomationScripts
 			}
 
 			// Actually execute UnrealPak
-			if (Commands.Count > 0)
+			if (PakCommands.Count > 0)
 			{
-				RunUnrealPakInParallel( Commands, LogNames, AdditionalCompressionOptionsOnCommandLine);
+				string PakCommandsFileName = CombinePaths(CmdEnv.LogFolder, "PakCommands.txt");
+				using (var Writer = new StreamWriter(PakCommandsFileName, false, new System.Text.UTF8Encoding(true)))
+				{
+					foreach (string Command in PakCommands)
+					{
+						Writer.WriteLine(Command);
+					}
+				}
+
+				StringBuilder Arguments = new StringBuilder(MakePathSafeToUseWithCommandLine(Params.RawProjectPath.FullName));
+				Arguments.AppendFormat(" {0}", GetCommonUnrealPakArguments(PrimaryOrderFiles, CommonAdditionalArgs, CryptoSettings, CryptoKeysCacheFilename, SecondaryOrderFiles, Params.Unattended));
+				Arguments.AppendFormat(" -CreateMultiple={0}", CommandUtils.MakePathSafeToUseWithCommandLine(PakCommandsFileName));
+				RunAndLog(CmdEnv, GetUnrealPakLocation().FullName, Arguments.ToString(), Options: ERunOptions.Default | ERunOptions.UTF8Output);
 			}
 
 			if (IoStoreCommands.Count > 0)
@@ -3396,49 +3425,6 @@ namespace AutomationScripts
 			RunAndLog(CmdEnv, GetUnrealPakLocation().FullName, CommandletParams, Options: ERunOptions.Default | ERunOptions.UTF8Output);
 		}
 
-		private static void RunUnrealPakInParallel(List<string> Commands, List<string> LogNames, string AdditionalCompressionOptionsOnCommandLine)
-		{
-			LogInformation("Executing {0} UnrealPak command{1}...", Commands.Count, (Commands.Count > 1)? "s" : "");
-
-			// Spawn tasks for each command
-			IProcessResult[] Results = new IProcessResult[Commands.Count];
-			string UnrealPakExe = GetUnrealPakLocation().FullName;
-			Parallel.ForEach(Commands, new ParallelOptions() { MaxDegreeOfParallelism = Environment.ProcessorCount }, (Command, LoopState, Idx) =>
-			{
-				int LocalIdx = (int)Idx;
-				string LogFile = LogUtils.GetUniqueLogName(CombinePaths(CmdEnv.LogFolder, String.Format("UnrealPak-{0}", LogNames[(int)Idx])));
-				string Arguments = Commands[LocalIdx] = String.Format("{0} -multiprocess -abslog={1} {2}", Commands[LocalIdx], MakePathSafeToUseWithCommandLine(LogFile), AdditionalCompressionOptionsOnCommandLine);
-
-				Results[LocalIdx] = Run(UnrealPakExe, Arguments, Options: ERunOptions.AppMustExist | ERunOptions.NoLoggingOfRunCommand);
-			});
-
-			// Output all the results
-			int NumFailed = 0;
-			for(int Idx = 0; Idx < Results.Length; Idx++)
-			{
-				LogInformation("Output from: {0}", Commands[Idx]);
-				using(new LogIndentScope("  "))
-				{
-					foreach(string Line in Results[Idx].Output.TrimEnd().Split('\n'))
-					{
-						LogInformation(Line);
-					}
-				}
-				LogInformation("UnrealPak terminated with exit code {0}", Results[Idx].ExitCode);
-
-				if(Results[Idx].ExitCode > 0)
-				{
-					NumFailed++;
-				}
-			}
-
-			// Abort if any instance failed
-			if(NumFailed > 0)
-			{
-				throw new AutomationException("UnrealPak failed");
-			}
-		}
-
 		/// <summary>
 		/// Creates pak files using streaming install chunk manifests.
 		/// </summary>
