[UnrealPak]
* Added -cryptofile documentation line when starting unrealpak without parameters

diff --git a/Engine/Source/Developer/PakFileUtilities/Private/PakFileUtilities.cpp b/Engine/Source/Developer/PakFileUtilities/Private/PakFileUtilities.cpp
--- a/Engine/Source/Developer/PakFileUtilities/Private/PakFileUtilities.cpp
+++ b/Engine/Source/Developer/PakFileUtilities/Private/PakFileUtilities.cpp
@@ -5743,6 +5743,7 @@ bool ExecuteUnrealPak(const TCHAR* CmdLine)
 	UE_LOG(LogPakFile, Error, TEXT("    -diff (requires 2 filenames first)"));
 	UE_LOG(LogPakFile, Error, TEXT("    -enginedir (specify engine dir for when using ini encryption configs)"));
 	UE_LOG(LogPakFile, Error, TEXT("    -projectdir (specify project dir for when using ini encryption configs)"));
+	UE_LOG(LogPakFile, Error, TEXT("    -cryptokeys=<cryptofile> (For example Crypto.json in cooked data)"));
 	UE_LOG(LogPakFile, Error, TEXT("    -encryptionini (specify ini base name to gather encryption settings from)"));
 	UE_LOG(LogPakFile, Error, TEXT("    -extracttomountpoint (Extract to mount point path of pak file)"));
 	UE_LOG(LogPakFile, Error, TEXT("    -encryptindex (encrypt the pak file index, making it unusable in unrealpak without supplying the key)"));
