Fixed unacceptable words in PakFileUtilities

diff --git a/Engine/Plugins/Experimental/DefaultInstallBundleManager/Source/Private/DefaultInstallBundleManager.cpp b/Engine/Plugins/Experimental/DefaultInstallBundleManager/Source/Private/DefaultInstallBundleManager.cpp
--- a/Engine/Plugins/Experimental/DefaultInstallBundleManager/Source/Private/DefaultInstallBundleManager.cpp
+++ b/Engine/Plugins/Experimental/DefaultInstallBundleManager/Source/Private/DefaultInstallBundleManager.cpp
@@ -182,12 +182,12 @@ void LoadKeychainFromInI(FKeyChain& OutCryptoSettings)
 			NewKey.Name = TEXT("Default");
 			NewKey.Guid = FGuid();
 			FMemory::Memcpy(NewKey.Key.Key, &Key[0], sizeof(FAES::FAESKey::Key));
-			OutCryptoSettings.EncryptionKeys.Add(NewKey.Guid, NewKey);
+			OutCryptoSettings.GetEncryptionKeys().Add(NewKey.Guid, NewKey);
 		}
 	}
 
 	FGuid EncryptionKeyOverrideGuid;
-	OutCryptoSettings.MasterEncryptionKey = OutCryptoSettings.EncryptionKeys.Find(EncryptionKeyOverrideGuid);
+	OutCryptoSettings.SetPrincipalEncryptionKey(OutCryptoSettings.GetEncryptionKeys().Find(EncryptionKeyOverrideGuid));
 }
 #endif // WITH_EDITOR
 
diff --git a/Engine/Source/Developer/IoStoreUtilities/Private/IoStoreUtilities.cpp b/Engine/Source/Developer/IoStoreUtilities/Private/IoStoreUtilities.cpp
--- a/Engine/Source/Developer/IoStoreUtilities/Private/IoStoreUtilities.cpp
+++ b/Engine/Source/Developer/IoStoreUtilities/Private/IoStoreUtilities.cpp
@@ -218,8 +218,8 @@ struct FReleasedPackages
 
 static void LoadKeyChain(const TCHAR* CmdLine, FKeyChain& OutCryptoSettings)
 {
-	OutCryptoSettings.SigningKey = InvalidRSAKeyHandle;
-	OutCryptoSettings.EncryptionKeys.Empty();
+	OutCryptoSettings.SetSigningKey(InvalidRSAKeyHandle);
+	OutCryptoSettings.GetEncryptionKeys().Empty();
 
 	// First, try and parse the keys from a supplied crypto key cache file
 	FString CryptoKeysCacheFilename;
@@ -281,7 +281,7 @@ static void LoadKeyChain(const TCHAR* CmdLine, FKeyChain& OutCryptoSettings)
 					FBase64::Decode(PrivateExpBase64, PrivateExp);
 					FBase64::Decode(ModulusBase64, Modulus);
 
-					OutCryptoSettings.SigningKey = FRSA::CreateKey(PublicExp, PrivateExp, Modulus);
+					OutCryptoSettings.SetSigningKey(FRSA::CreateKey(PublicExp, PrivateExp, Modulus));
 
 					UE_LOG(LogIoStore, Display, TEXT("Parsed signature keys from config files."));
 				}
@@ -300,7 +300,7 @@ static void LoadKeyChain(const TCHAR* CmdLine, FKeyChain& OutCryptoSettings)
 						NewKey.Name = TEXT("Default");
 						NewKey.Guid = FGuid();
 						FMemory::Memcpy(NewKey.Key.Key, &Key[0], sizeof(FAES::FAESKey::Key));
-						OutCryptoSettings.EncryptionKeys.Add(NewKey.Guid, NewKey);
+						OutCryptoSettings.GetEncryptionKeys().Add(NewKey.Guid, NewKey);
 						UE_LOG(LogIoStore, Display, TEXT("Parsed AES encryption key from config files."));
 					}
 				}
@@ -344,7 +344,7 @@ static void LoadKeyChain(const TCHAR* CmdLine, FKeyChain& OutCryptoSettings)
 						{
 							NewKey.Key.Key[Index] = (uint8)EncryptionKeyString[Index];
 						}
-						OutCryptoSettings.EncryptionKeys.Add(NewKey.Guid, NewKey);
+						OutCryptoSettings.GetEncryptionKeys().Add(NewKey.Guid, NewKey);
 						UE_LOG(LogIoStore, Display, TEXT("Parsed AES encryption key from config files."));
 					}
 				}
@@ -387,7 +387,7 @@ static void LoadKeyChain(const TCHAR* CmdLine, FKeyChain& OutCryptoSettings)
 			const auto AsAnsi = StringCast<ANSICHAR>(*EncryptionKeyString);
 			check(AsAnsi.Length() == RequiredKeyLength);
 			FMemory::Memcpy(NewKey.Key.Key, AsAnsi.Get(), RequiredKeyLength);
-			OutCryptoSettings.EncryptionKeys.Add(NewKey.Guid, NewKey);
+			OutCryptoSettings.GetEncryptionKeys().Add(NewKey.Guid, NewKey);
 			UE_LOG(LogIoStore, Display, TEXT("Parsed AES encryption key from command line."));
 		}
 	}
@@ -399,7 +399,7 @@ static void LoadKeyChain(const TCHAR* CmdLine, FKeyChain& OutCryptoSettings)
 		UE_LOG(LogIoStore, Display, TEXT("Using encryption key override '%s'"), *EncryptionKeyOverrideGuidString);
 		FGuid::Parse(EncryptionKeyOverrideGuidString, EncryptionKeyOverrideGuid);
 	}
-	OutCryptoSettings.MasterEncryptionKey = OutCryptoSettings.EncryptionKeys.Find(EncryptionKeyOverrideGuid);
+	OutCryptoSettings.SetPrincipalEncryptionKey(OutCryptoSettings.GetEncryptionKeys().Find(EncryptionKeyOverrideGuid));
 }
 
 struct FContainerSourceFile 
@@ -1637,7 +1637,7 @@ static TUniquePtr<FIoStoreReader> CreateIoStoreReader(const TCHAR* Path, const F
 	TUniquePtr<FIoStoreReader> IoStoreReader(new FIoStoreReader());
 
 	TMap<FGuid, FAES::FAESKey> DecryptionKeys;
-	for (const auto& KV : KeyChain.EncryptionKeys)
+	for (const auto& KV : KeyChain.GetEncryptionKeys())
 	{
 		DecryptionKeys.Add(KV.Key, KV.Value.Key);
 	}
@@ -3323,7 +3323,7 @@ int32 CreateTarget(const FIoStoreArguments& Arguments, const FIoStoreWriterSetti
 			FIoContainerSettings GlobalContainerSettings;
 			if (Arguments.bSign)
 			{
-				GlobalContainerSettings.SigningKey = Arguments.KeyChain.SigningKey;
+				GlobalContainerSettings.SigningKey = Arguments.KeyChain.GetSigningKey();
 				GlobalContainerSettings.ContainerFlags |= EIoContainerFlags::Signed;
 			}
 			GlobalIoStoreWriter = IoStoreWriterContext->CreateContainer(*Arguments.GlobalContainerPath, GlobalContainerSettings);
@@ -3342,14 +3342,14 @@ int32 CreateTarget(const FIoStoreArguments& Arguments, const FIoStoreWriterSetti
 				}
 				if (EnumHasAnyFlags(ContainerTarget->ContainerFlags, EIoContainerFlags::Encrypted))
 				{
-					const FNamedAESKey* Key = Arguments.KeyChain.EncryptionKeys.Find(ContainerTarget->EncryptionKeyGuid);
+					const FNamedAESKey* Key = Arguments.KeyChain.GetEncryptionKeys().Find(ContainerTarget->EncryptionKeyGuid);
 					check(Key);
 					ContainerSettings.EncryptionKeyGuid = ContainerTarget->EncryptionKeyGuid;
 					ContainerSettings.EncryptionKey = Key->Key;
 				}
 				if (EnumHasAnyFlags(ContainerTarget->ContainerFlags, EIoContainerFlags::Signed))
 				{
-					ContainerSettings.SigningKey = Arguments.KeyChain.SigningKey;
+					ContainerSettings.SigningKey = Arguments.KeyChain.GetSigningKey();
 					ContainerSettings.ContainerFlags |= EIoContainerFlags::Signed;
 				}
 				ContainerSettings.bGenerateDiffPatch = ContainerTarget->bGenerateDiffPatch;
@@ -3777,14 +3777,14 @@ int32 CreateContentPatch(const FIoStoreArguments& Arguments, const FIoStoreWrite
 		ContainerSettings.ContainerId = TargetReader->GetContainerId();
 		if (Arguments.bSign || EnumHasAnyFlags(TargetContainerFlags, EIoContainerFlags::Signed))
 		{
-			ContainerSettings.SigningKey = Arguments.KeyChain.SigningKey;
+			ContainerSettings.SigningKey =Arguments.KeyChain.GetSigningKey();
 			ContainerSettings.ContainerFlags |= EIoContainerFlags::Signed;
 		}
 
 		if (EnumHasAnyFlags(TargetContainerFlags, EIoContainerFlags::Encrypted))
 		{
 			ContainerSettings.ContainerFlags |= EIoContainerFlags::Encrypted;
-			const FNamedAESKey* Key = Arguments.KeyChain.EncryptionKeys.Find(TargetReader->GetEncryptionKeyGuid());
+			const FNamedAESKey* Key = Arguments.KeyChain.GetEncryptionKeys().Find(TargetReader->GetEncryptionKeyGuid());
 			if (!Key)
 			{
 				UE_LOG(LogIoStore, Error, TEXT("Missing encryption key for target container"));
@@ -5517,7 +5517,7 @@ int32 Staged2Zen(const FString& BuildPath, const FKeyChain& KeyChain, const FStr
 
 	UE_LOG(LogIoStore, Display, TEXT("Extracting files from paks..."));
 	FPakPlatformFile PakPlatformFile;
-	for (const auto& KV : KeyChain.EncryptionKeys)
+	for (const auto& KV : KeyChain.GetEncryptionKeys())
 	{
 		FCoreDelegates::GetRegisterEncryptionKeyMulticastDelegate().Broadcast(KV.Key, KV.Value.Key);
 	}
diff --git a/Engine/Source/Developer/PakFileUtilities/Private/PakFileUtilities.cpp b/Engine/Source/Developer/PakFileUtilities/Private/PakFileUtilities.cpp
--- a/Engine/Source/Developer/PakFileUtilities/Private/PakFileUtilities.cpp
+++ b/Engine/Source/Developer/PakFileUtilities/Private/PakFileUtilities.cpp
@@ -865,8 +865,8 @@ bool PrepareCopyFileToPak(const FString& InMountPoint, const FPakInputPair& InFi
 			}
 
 			//Encrypt the buffer before writing it to disk
-			check(InKeyChain.MasterEncryptionKey);
-			FAES::EncryptData(UncompressedBuffer, PaddedEncryptedFileSize, InKeyChain.MasterEncryptionKey->Key);
+			check(InKeyChain.GetPrincipalEncryptionKey());
+			FAES::EncryptData(UncompressedBuffer, PaddedEncryptedFileSize, InKeyChain.GetPrincipalEncryptionKey()->Key);
 			// Update the size to be written
 			OutSizeToWrite = PaddedEncryptedFileSize;
 			OutNewEntry.Info.SetEncrypted( true );
@@ -917,8 +917,8 @@ bool PrepareCopyCompressedFileToPak(const FString& InMountPoint, FPakInfo& Info,
 
 	if (InFile.bNeedEncryption)
 	{
-		check(InKeyChain.MasterEncryptionKey);
-		FAES::EncryptData(CompressedFile.CompressedBuffer.Get(), CompressedFile.TotalCompressedSize, InKeyChain.MasterEncryptionKey->Key);
+		check(InKeyChain.GetPrincipalEncryptionKey());
+		FAES::EncryptData(CompressedFile.CompressedBuffer.Get(), CompressedFile.TotalCompressedSize, InKeyChain.GetPrincipalEncryptionKey()->Key);
 	}
 
 	//Hash the final buffer thats written
@@ -1420,7 +1420,7 @@ bool BufferedCopyFile(FArchive& Dest, FArchive& Source, const FPakFile& PakFile,
 		Source.Serialize(Buffer,SizeToRead);
 		if (Entry.IsEncrypted())
 		{
-			const FNamedAESKey* Key = InKeyChain.MasterEncryptionKey;
+			const FNamedAESKey* Key = InKeyChain.GetPrincipalEncryptionKey();
 			check(Key);
 			FAES::DecryptData((uint8*)Buffer, SizeToRead, Key->Key);
 		}
@@ -1466,10 +1466,10 @@ bool UncompressCopyFile(FArchive& Dest, FArchive& Source, const FPakEntry& Entry
 
 		if (Entry.IsEncrypted())
 		{
-			const FNamedAESKey* Key = InKeyChain.EncryptionKeys.Find(PakFile.GetInfo().EncryptionKeyGuid);
+			const FNamedAESKey* Key = InKeyChain.GetEncryptionKeys().Find(PakFile.GetInfo().EncryptionKeyGuid);
 			if (Key == nullptr)
 			{
-				Key = InKeyChain.MasterEncryptionKey;
+				Key = InKeyChain.GetPrincipalEncryptionKey();
 			}
 			check(Key);
 			FAES::DecryptData(PersistentBuffer, SizeToRead, Key->Key);
@@ -1503,8 +1503,8 @@ TEncryptionInt ParseEncryptionIntFromJson(TSharedPtr<FJsonObject> InObj, const T
 
 void LoadKeyChain(const TCHAR* CmdLine, FKeyChain& OutCryptoSettings)
 {
-	OutCryptoSettings.SigningKey = InvalidRSAKeyHandle;
-	OutCryptoSettings.EncryptionKeys.Empty();
+	OutCryptoSettings.SetSigningKey( InvalidRSAKeyHandle );
+	OutCryptoSettings.GetEncryptionKeys().Empty();
 
 	// First, try and parse the keys from a supplied crypto key cache file
 	FString CryptoKeysCacheFilename;
@@ -1566,7 +1566,7 @@ void LoadKeyChain(const TCHAR* CmdLine, FKeyChain& OutCryptoSettings)
 					FBase64::Decode(PrivateExpBase64, PrivateExp);
 					FBase64::Decode(ModulusBase64, Modulus);
 
-					OutCryptoSettings.SigningKey = FRSA::CreateKey(PublicExp, PrivateExp, Modulus);
+					OutCryptoSettings.SetSigningKey(FRSA::CreateKey(PublicExp, PrivateExp, Modulus));
 
 					UE_LOG(LogPakFile, Display, TEXT("Parsed signature keys from config files."));
 				}
@@ -1585,7 +1585,7 @@ void LoadKeyChain(const TCHAR* CmdLine, FKeyChain& OutCryptoSettings)
 						NewKey.Name = TEXT("Default");
 						NewKey.Guid = FGuid();
 						FMemory::Memcpy(NewKey.Key.Key, &Key[0], sizeof(FAES::FAESKey::Key));
-						OutCryptoSettings.EncryptionKeys.Add(NewKey.Guid, NewKey);
+						OutCryptoSettings.GetEncryptionKeys().Add(NewKey.Guid, NewKey);
 						UE_LOG(LogPakFile, Display, TEXT("Parsed AES encryption key from config files."));
 					}
 				}
@@ -1629,7 +1629,7 @@ void LoadKeyChain(const TCHAR* CmdLine, FKeyChain& OutCryptoSettings)
 						{
 							NewKey.Key.Key[Index] = (uint8)EncryptionKeyString[Index];
 						}
-						OutCryptoSettings.EncryptionKeys.Add(NewKey.Guid, NewKey);
+						OutCryptoSettings.GetEncryptionKeys().Add(NewKey.Guid, NewKey);
 						UE_LOG(LogPakFile, Display, TEXT("Parsed AES encryption key from config files."));
 					}
 				}
@@ -1672,7 +1672,7 @@ void LoadKeyChain(const TCHAR* CmdLine, FKeyChain& OutCryptoSettings)
 			const auto AsAnsi = StringCast<ANSICHAR>(*EncryptionKeyString);
 			check(AsAnsi.Length() == RequiredKeyLength);
 			FMemory::Memcpy(NewKey.Key.Key, AsAnsi.Get(), RequiredKeyLength);
-			OutCryptoSettings.EncryptionKeys.Add(NewKey.Guid, NewKey);
+			OutCryptoSettings.GetEncryptionKeys().Add(NewKey.Guid, NewKey);
 			UE_LOG(LogPakFile, Display, TEXT("Parsed AES encryption key from command line."));
 		}
 	}
@@ -1683,7 +1683,7 @@ void LoadKeyChain(const TCHAR* CmdLine, FKeyChain& OutCryptoSettings)
 	{
 		FGuid::Parse(EncryptionKeyOverrideGuidString, EncryptionKeyOverrideGuid);
 	}
-	OutCryptoSettings.MasterEncryptionKey = OutCryptoSettings.EncryptionKeys.Find(EncryptionKeyOverrideGuid);
+	OutCryptoSettings.SetPrincipalEncryptionKey(OutCryptoSettings.GetEncryptionKeys().Find(EncryptionKeyOverrideGuid));
 }
 
 /**
@@ -1698,7 +1698,7 @@ FArchive* CreatePakWriter(const TCHAR* Filename, const FKeyChain& InKeyChain, bo
 		if (bSign)
 		{
 			UE_LOG(LogPakFile, Display, TEXT("Creating signed pak %s."), Filename);
-			Writer = new FSignedArchiveWriter(*Writer, Filename, InKeyChain.SigningKey);
+			Writer = new FSignedArchiveWriter(*Writer, Filename, InKeyChain.GetSigningKey());
 		}
 		else
 		{
@@ -2005,13 +2005,13 @@ bool FPakWriterContext::AddPakFile(const TCHAR* Filename, const TArray<FPakInput
 		}
 	}
 
-	if (InKeyChain.MasterEncryptionKey)
+	if (InKeyChain.GetPrincipalEncryptionKey())
 	{
-		UE_LOG(LogPakFile, Display, TEXT("Using encryption key '%s' [%s]"), *InKeyChain.MasterEncryptionKey->Name, *InKeyChain.MasterEncryptionKey->Guid.ToString());
+		UE_LOG(LogPakFile, Display, TEXT("Using encryption key '%s' [%s]"), *InKeyChain.GetPrincipalEncryptionKey()->Name, *InKeyChain.GetPrincipalEncryptionKey()->Guid.ToString());
 	}
 
-	OutputPakFile->Info.bEncryptedIndex = (InKeyChain.MasterEncryptionKey && CmdLineParameters.EncryptIndex);
-	OutputPakFile->Info.EncryptionKeyGuid = InKeyChain.MasterEncryptionKey ? InKeyChain.MasterEncryptionKey->Guid : FGuid();
+	OutputPakFile->Info.bEncryptedIndex = (InKeyChain.GetPrincipalEncryptionKey() && CmdLineParameters.EncryptIndex);
+	OutputPakFile->Info.EncryptionKeyGuid = InKeyChain.GetPrincipalEncryptionKey() ? InKeyChain.GetPrincipalEncryptionKey()->Guid : FGuid();
 
 	if (CmdLineParameters.bPatchCompatibilityMode421)
 	{
@@ -2450,7 +2450,7 @@ void FPakWriterContext::WriterThreadFunc()
 			{
 				OutputEntry->PakFile->TotalRequestedEncryptedFiles++;
 
-				if (OutputEntry->PakFile->KeyChain.MasterEncryptionKey)
+				if (OutputEntry->PakFile->KeyChain.GetPrincipalEncryptionKey())
 				{
 					OutputEntry->PakFile->TotalEncryptedFiles++;
 					OutputEntry->PakFile->TotalEncryptedDataSize += SizeToWrite;
@@ -2506,7 +2506,7 @@ bool FPakWriterContext::Flush()
 
 				if (OutputEntry.InputPair.bNeedEncryption)
 				{
-					if (OutputEntry.PakFile->KeyChain.MasterEncryptionKey)
+					if (OutputEntry.PakFile->KeyChain.GetPrincipalEncryptionKey())
 					{
 						EncryptedString = TEXT("encrypted ");
 					}
@@ -2710,8 +2710,8 @@ void WritePakFooter(FArchive& PakHandle, FPakFooterInfo& Footer)
 
 		if (Footer.Info.bEncryptedIndex)
 		{
-			check(Footer.KeyChain && Footer.KeyChain->MasterEncryptionKey && Footer.TotalEncryptedDataSize);
-			FAES::EncryptData(IndexData.GetData(), IndexData.Num(), Footer.KeyChain->MasterEncryptionKey->Key);
+			check(Footer.KeyChain && Footer.KeyChain->GetPrincipalEncryptionKey() && Footer.TotalEncryptedDataSize);
+			FAES::EncryptData(IndexData.GetData(), IndexData.Num(), Footer.KeyChain->GetPrincipalEncryptionKey()->Key);
 			*Footer.TotalEncryptedDataSize += IndexData.Num();
 		}
 	};
@@ -3011,7 +3011,7 @@ bool ListFilesInPak(const TCHAR * InPakFilename, int64 SizeFilter, bool bInclude
 	}
 	else
 	{
-		if (PakFile.GetInfo().Magic != 0 && PakFile.GetInfo().EncryptionKeyGuid.IsValid() && !InKeyChain.EncryptionKeys.Contains(PakFile.GetInfo().EncryptionKeyGuid))
+		if (PakFile.GetInfo().Magic != 0 && PakFile.GetInfo().EncryptionKeyGuid.IsValid() && !InKeyChain.GetEncryptionKeys().Contains(PakFile.GetInfo().EncryptionKeyGuid))
 		{
 			UE_LOG(LogPakFile, Fatal, TEXT("Missing encryption key %s for pak file \"%s\"."), *PakFile.GetInfo().EncryptionKeyGuid.ToString(), InPakFilename);
 		}
@@ -3501,7 +3501,7 @@ bool ShowCompressionBlockCRCs( const TCHAR* InPakFileName, TArray<int64>& InOffs
 
 			if (Entry->IsEncrypted())
 			{
-				const FNamedAESKey* Key = InKeyChain.MasterEncryptionKey;
+				const FNamedAESKey* Key = InKeyChain.GetPrincipalEncryptionKey();
 				check(Key);
 				FAES::DecryptData(PersistentBuffer, SizeToRead, Key->Key);
 			}
@@ -4895,10 +4895,10 @@ bool Repack(const FString& InputPakFile, const FString& OutputPakFile, const FPa
 		FString TempOutputPakFile = FPaths::CreateTempFilename(*FPaths::GetPath(OutputPakFile), *FPaths::GetCleanFilename(OutputPakFile));
 
 		FPakCommandLineParameters ModifiedCmdLineParameters = CmdLineParameters;
-		ModifiedCmdLineParameters.bSign = bAnySigned && (InKeyChain.SigningKey != InvalidRSAKeyHandle);
+		ModifiedCmdLineParameters.bSign = bAnySigned && (InKeyChain.GetSigningKey() != InvalidRSAKeyHandle);
 		
 		FKeyChain ModifiedKeyChain = InKeyChain;
-		ModifiedKeyChain.MasterEncryptionKey = InKeyChain.EncryptionKeys.Find(EncryptionKeys.Num() ? EncryptionKeys[0] : FGuid());
+		ModifiedKeyChain.SetPrincipalEncryptionKey( InKeyChain.GetEncryptionKeys().Find(EncryptionKeys.Num() ? EncryptionKeys[0] : FGuid()) );
 
 		// Create the new pak file
 		UE_LOG(LogPakFile, Display, TEXT("Creating %s..."), *OutputPakFile);
@@ -5615,7 +5615,7 @@ bool ExecuteUnrealPak(const TCHAR* CmdLine)
 			{
 				FGuid::Parse(EncryptionKeyOverrideGuidString, EncryptionKeyOverrideGuid);
 			}
-			KeyChainForPakFile.MasterEncryptionKey = KeyChainForPakFile.EncryptionKeys.Find(EncryptionKeyOverrideGuid);
+			KeyChainForPakFile.SetPrincipalEncryptionKey( KeyChainForPakFile.GetEncryptionKeys().Find(EncryptionKeyOverrideGuid) );
 		
 			int32 LowestSourcePakVersion = 0;
 			TMap<FString, FFileInfo> SourceFileHashes;
@@ -5652,7 +5652,7 @@ bool ExecuteUnrealPak(const TCHAR* CmdLine)
 					if (UsedEncryptionKeys[0].IsValid())
 					{
 						UE_LOG(LogPakFile, Display, TEXT("Found encryption key %s in pak file %s will use to encrypt patch"), *((UsedEncryptionKeys[0]).ToString()), *CmdLineParameters.SourcePatchPakFilename);
-						KeyChainForPakFile.MasterEncryptionKey = KeyChainForPakFile.EncryptionKeys.Find(UsedEncryptionKeys[0]);
+						KeyChainForPakFile.SetPrincipalEncryptionKey(KeyChainForPakFile.GetEncryptionKeys().Find(UsedEncryptionKeys[0]));
 					}
 				}
 			}
diff --git a/Engine/Source/Developer/PakFileUtilities/Private/Tests/KeyChainTest.cpp b/Engine/Source/Developer/PakFileUtilities/Private/Tests/KeyChainTest.cpp
--- a/Engine/Source/Developer/PakFileUtilities/Private/Tests/KeyChainTest.cpp
+++ b/Engine/Source/Developer/PakFileUtilities/Private/Tests/KeyChainTest.cpp
@@ -19,130 +19,130 @@ bool FKeyChainTest::RunTest(const FString& Parameters)
 	// Default construct
 	{
 		FKeyChain Chain;
-		TestTrue(TEXT("Default construct - master key is invalid"), Chain.MasterEncryptionKey == nullptr); 
-		TestTrue(TEXT("Default construct - signing key is invalid"), Chain.SigningKey == nullptr); 
-		TestTrue(TEXT("Default construct - encryption keys are empty"), Chain.EncryptionKeys.Num() == 0); 
+		TestTrue(TEXT("Default construct - principal key is invalid"), Chain.GetPrincipalEncryptionKey() == nullptr);
+		TestTrue(TEXT("Default construct - signing key is invalid"), Chain.GetSigningKey() == nullptr);
+		TestTrue(TEXT("Default construct - encryption keys are empty"), Chain.GetEncryptionKeys().Num() == 0);
 	}
 
 	// Copy construct
 	{
 		FKeyChain Chain;
-		Chain.EncryptionKeys.Add(DefaultGuid, FNamedAESKey { TEXT("Default"), DefaultGuid, FAES::FAESKey() });
-		Chain.MasterEncryptionKey = Chain.EncryptionKeys.Find(DefaultGuid);
-		Chain.SigningKey = DefaultSigningKey;
+		Chain.GetEncryptionKeys().Add(DefaultGuid, FNamedAESKey { TEXT("Default"), DefaultGuid, FAES::FAESKey() });
+		Chain.SetPrincipalEncryptionKey(Chain.GetEncryptionKeys().Find(DefaultGuid));
+		Chain.SetSigningKey(DefaultSigningKey);
 
 		FKeyChain Copy(Chain);
 
-		TestTrue(TEXT("Copy construct - with valid master key does NOT copy pointer"),
-			Copy.MasterEncryptionKey != Chain.MasterEncryptionKey); 
-		TestTrue(TEXT("Copy construct - master key name matches"),
-			Copy.MasterEncryptionKey->Name == Chain.MasterEncryptionKey->Name); 
-		TestTrue(TEXT("Copy construct - master key GUID name matches"),
-			Copy.MasterEncryptionKey->Guid == Chain.MasterEncryptionKey->Guid); 
-		TestTrue(TEXT("Copy construct - with valid master key sets master key"),
-			Copy.MasterEncryptionKey == Copy.EncryptionKeys.Find(DefaultGuid)); 
+		TestTrue(TEXT("Copy construct - with valid principal key does NOT copy pointer"),
+			Copy.GetPrincipalEncryptionKey() != Chain.GetPrincipalEncryptionKey());
+		TestTrue(TEXT("Copy construct - principal key name matches"),
+			Copy.GetPrincipalEncryptionKey()->Name == Chain.GetPrincipalEncryptionKey()->Name);
+		TestTrue(TEXT("Copy construct - principal key GUID name matches"),
+			Copy.GetPrincipalEncryptionKey()->Guid == Chain.GetPrincipalEncryptionKey()->Guid);
+		TestTrue(TEXT("Copy construct - with valid principal key sets principal key"),
+			Copy.GetPrincipalEncryptionKey() == Copy.GetEncryptionKeys().Find(DefaultGuid));
 		TestTrue(TEXT("Copy construct - signing key matches"),
-			Copy.SigningKey == Chain.SigningKey ); 
+			Copy.GetSigningKey() == Chain.GetSigningKey());
 	}
 
 	// Copy assign
 	{
 		{
 			FKeyChain Chain;
-			Chain.EncryptionKeys.Add(DefaultGuid, FNamedAESKey { TEXT("Default"), DefaultGuid, FAES::FAESKey() });
-			Chain.MasterEncryptionKey = Chain.EncryptionKeys.Find(DefaultGuid);
-			Chain.SigningKey = DefaultSigningKey;
+			Chain.GetEncryptionKeys().Add(DefaultGuid, FNamedAESKey{TEXT("Default"), DefaultGuid, FAES::FAESKey()});
+			Chain.SetPrincipalEncryptionKey(Chain.GetEncryptionKeys().Find(DefaultGuid));
+			Chain.SetSigningKey(DefaultSigningKey);
 
 			FKeyChain Copy;
 			Copy = Chain;
 			
-			TestTrue(TEXT("Copy assign - with valid master key does NOT copy pointer"),
-				Copy.MasterEncryptionKey != Chain.MasterEncryptionKey); 
-			TestTrue(TEXT("Copy assign - master key name matches"),
-				Copy.MasterEncryptionKey->Name == Chain.MasterEncryptionKey->Name); 
-			TestTrue(TEXT("Copy assign - master key GUID name matches"),
-				Copy.MasterEncryptionKey->Guid == Chain.MasterEncryptionKey->Guid); 
-			TestTrue(TEXT("Copy assign - with valid master key sets master key"),
-				Copy.MasterEncryptionKey == Copy.EncryptionKeys.Find(DefaultGuid)); 
+			TestTrue(TEXT("Copy assign - with valid principal key does NOT copy pointer"),
+				Copy.GetPrincipalEncryptionKey() != Chain.GetPrincipalEncryptionKey());
+			TestTrue(TEXT("Copy assign - principal key name matches"),
+				Copy.GetPrincipalEncryptionKey()->Name == Chain.GetPrincipalEncryptionKey()->Name);
+			TestTrue(TEXT("Copy assign - principal key GUID name matches"),
+				Copy.GetPrincipalEncryptionKey()->Guid == Chain.GetPrincipalEncryptionKey()->Guid);
+			TestTrue(TEXT("Copy assign - with valid principal key sets principal key"),
+				Copy.GetPrincipalEncryptionKey() == Copy.GetEncryptionKeys().Find(DefaultGuid));
 			TestTrue(TEXT("Copy assign - signing key matches"),
-				Copy.SigningKey == Chain.SigningKey ); 
+				Copy.GetSigningKey() == Chain.GetSigningKey());
 		}
 
 		{
 			FKeyChain Copy;
-			Copy.EncryptionKeys.Add(DefaultGuid, FNamedAESKey { TEXT("Default"), DefaultGuid, FAES::FAESKey() });
-			Copy.MasterEncryptionKey = Copy.EncryptionKeys.Find(DefaultGuid);
-			Copy.SigningKey = DefaultSigningKey;
+			Copy.GetEncryptionKeys().Add(DefaultGuid, FNamedAESKey { TEXT("Default"), DefaultGuid, FAES::FAESKey() });
+			Copy.SetPrincipalEncryptionKey(Copy.GetEncryptionKeys().Find(DefaultGuid));
+			Copy.SetSigningKey(DefaultSigningKey);
 
 			FKeyChain Chain;
 			Copy = Chain;
 			
-			TestTrue(TEXT("Copy assign - empty instance, clears master key"),
-				Copy.MasterEncryptionKey == nullptr); 
+			TestTrue(TEXT("Copy assign - empty instance, clears principal key"),
+				Copy.GetPrincipalEncryptionKey() == nullptr);
 			TestTrue(TEXT("Copy assign - empty instance, clears encryption keys"),
-				Copy.EncryptionKeys.Num() == 0); 
+				Copy.GetEncryptionKeys().Num() == 0);
 			TestTrue(TEXT("Copy assign - signing key is invalid"),
-				Copy.SigningKey == nullptr); 
+				Copy.GetSigningKey() == nullptr);
 		}
 	}
 
 	// Move construct
 	{
 		FKeyChain Moved;
-		Moved.EncryptionKeys.Add(DefaultGuid, FNamedAESKey { TEXT("Default"), DefaultGuid, FAES::FAESKey() });
-		Moved.MasterEncryptionKey = Moved.EncryptionKeys.Find(DefaultGuid);
-		Moved.SigningKey = DefaultSigningKey;
+		Moved.GetEncryptionKeys().Add(DefaultGuid, FNamedAESKey { TEXT("Default"), DefaultGuid, FAES::FAESKey() });
+		Moved.SetPrincipalEncryptionKey(Moved.GetEncryptionKeys().Find(DefaultGuid));
+		Moved.SetSigningKey(DefaultSigningKey);
 
 		FKeyChain Chain(MoveTemp(Moved));
 
-		TestTrue(TEXT("Move construct - with valid master key sets master key"),
-			Chain.MasterEncryptionKey == Chain.EncryptionKeys.Find(DefaultGuid)); 
-		TestTrue(TEXT("Move construct - with valid master key sets signing key"),
-			Chain.SigningKey == DefaultSigningKey); 
+		TestTrue(TEXT("Move construct - with valid principal key sets principal key"),
+			Chain.GetPrincipalEncryptionKey() == Chain.GetEncryptionKeys().Find(DefaultGuid));
+		TestTrue(TEXT("Move construct - with valid principal key sets signing key"),
+			Chain.GetSigningKey() == DefaultSigningKey); 
 		TestTrue(TEXT("Move construct - invalidates moved instance"),
-			Moved.MasterEncryptionKey == nullptr); 
+			Moved.GetPrincipalEncryptionKey() == nullptr);
 		TestTrue(TEXT("Move construct - invalidates moved instance"),
-			Moved.SigningKey == InvalidRSAKeyHandle); 
+			Moved.GetSigningKey() == InvalidRSAKeyHandle);
 		TestTrue(TEXT("Move construct - invalidates moved instance"),
-			Moved.EncryptionKeys.Num() == 0); 
+			Moved.GetEncryptionKeys().Num() == 0);
 	}
 
 	// Move assign
 	{
 		{
 			FKeyChain Moved;
-			Moved.EncryptionKeys.Add(DefaultGuid, FNamedAESKey { TEXT("Default"), DefaultGuid, FAES::FAESKey() });
-			Moved.MasterEncryptionKey = Moved.EncryptionKeys.Find(DefaultGuid);
-			Moved.SigningKey = DefaultSigningKey;
+			Moved.GetEncryptionKeys().Add(DefaultGuid, FNamedAESKey{TEXT("Default"), DefaultGuid, FAES::FAESKey()});
+			Moved.SetPrincipalEncryptionKey(Moved.GetEncryptionKeys().Find(DefaultGuid));
+			Moved.SetSigningKey(DefaultSigningKey);
 
 			FKeyChain Chain;
 			Chain = MoveTemp(Moved);
 
-			TestTrue(TEXT("Move construct - with valid master key sets master key"),
-				Chain.MasterEncryptionKey == Chain.EncryptionKeys.Find(DefaultGuid)); 
-			TestTrue(TEXT("Move construct - with valid master key sets signing key"),
-				Chain.SigningKey == DefaultSigningKey); 
+			TestTrue(TEXT("Move construct - with valid principal key sets principal key"),
+				Chain.GetPrincipalEncryptionKey() == Chain.GetEncryptionKeys().Find(DefaultGuid));
+			TestTrue(TEXT("Move construct - with valid principal key sets signing key"),
+				Chain.GetSigningKey() == DefaultSigningKey);
 			TestTrue(TEXT("Move construct - invalidates moved instance"),
-				Moved.MasterEncryptionKey == nullptr); 
+				Moved.GetPrincipalEncryptionKey() == nullptr);
 			TestTrue(TEXT("Move construct - invalidates moved instance"),
-				Moved.SigningKey == InvalidRSAKeyHandle); 
+				Moved.GetSigningKey() == InvalidRSAKeyHandle);
 			TestTrue(TEXT("Move construct - invalidates moved instance"),
-				Moved.EncryptionKeys.Num() == 0); 
+				Moved.GetEncryptionKeys().Num() == 0);
 		}
 		
 		{
 			FKeyChain Copy;
-			Copy.EncryptionKeys.Add(DefaultGuid, FNamedAESKey { TEXT("Default"), DefaultGuid, FAES::FAESKey() });
-			Copy.MasterEncryptionKey = Copy.EncryptionKeys.Find(DefaultGuid);
+			Copy.GetEncryptionKeys().Add(DefaultGuid, FNamedAESKey{TEXT("Default"), DefaultGuid, FAES::FAESKey()});
+			Copy.SetPrincipalEncryptionKey(Copy.GetEncryptionKeys().Find(DefaultGuid));
 
 			Copy = FKeyChain();
 
-			TestTrue(TEXT("Move assign - empty instance, clears master key"),
-				Copy.MasterEncryptionKey == nullptr); 
+			TestTrue(TEXT("Move assign - empty instance, clears principal key"),
+				Copy.GetPrincipalEncryptionKey() == nullptr);
 			TestTrue(TEXT("Move assign - empty instance, clears encryption keys"),
-				Copy.EncryptionKeys.Num() == 0);
+				Copy.GetEncryptionKeys().Num() == 0);
 			TestTrue(TEXT("Move assign - invalidates signing key"),
-				Copy.SigningKey == InvalidRSAKeyHandle); 
+				Copy.GetSigningKey() == InvalidRSAKeyHandle);
 		}
 	}
 
diff --git a/Engine/Source/Runtime/Core/Public/Misc/KeyChainUtilities.h b/Engine/Source/Runtime/Core/Public/Misc/KeyChainUtilities.h
--- a/Engine/Source/Runtime/Core/Public/Misc/KeyChainUtilities.h
+++ b/Engine/Source/Runtime/Core/Public/Misc/KeyChainUtilities.h
@@ -26,43 +26,47 @@ struct FNamedAESKey
 
 struct FKeyChain
 {
+public:
+
 	FKeyChain() = default;
 
 	FKeyChain(const FKeyChain& Other)
-		: SigningKey(Other.SigningKey)
-		, EncryptionKeys(Other.EncryptionKeys)
 	{
-		if (Other.MasterEncryptionKey)
+		SetSigningKey(Other.GetSigningKey());
+		SetEncryptionKeys(Other.GetEncryptionKeys());
+
+		if (Other.GetPrincipalEncryptionKey())
 		{
-			MasterEncryptionKey = EncryptionKeys.Find(Other.MasterEncryptionKey->Guid);
+			SetPrincipalEncryptionKey(GetEncryptionKeys().Find(Other.GetPrincipalEncryptionKey()->Guid));
 		}
 	}
 	
 	FKeyChain(FKeyChain&& Other)
-		: SigningKey(Other.SigningKey)
-		, EncryptionKeys(MoveTemp(Other.EncryptionKeys))
 	{
-		if (Other.MasterEncryptionKey)
+		SetSigningKey(Other.GetSigningKey());
+		SetEncryptionKeys(MoveTemp(Other.GetEncryptionKeys()));
+
+		if (Other.GetPrincipalEncryptionKey())
 		{
-			MasterEncryptionKey = EncryptionKeys.Find(Other.MasterEncryptionKey->Guid);
+			SetPrincipalEncryptionKey(GetEncryptionKeys().Find(Other.GetPrincipalEncryptionKey()->Guid));
 		}
 		
-		Other.SigningKey = InvalidRSAKeyHandle;
-		Other.MasterEncryptionKey = nullptr;
+		Other.SetSigningKey(InvalidRSAKeyHandle);
+		Other.SetPrincipalEncryptionKey(nullptr);
 	}
 
 	FKeyChain& operator=(const FKeyChain& Other)
 	{
-		SigningKey = Other.SigningKey;
-		EncryptionKeys = Other.EncryptionKeys;
+		SetSigningKey(Other.GetSigningKey());
+		SetEncryptionKeys(Other.GetEncryptionKeys());
 		
-		if (Other.MasterEncryptionKey)
+		if (Other.GetPrincipalEncryptionKey())
 		{
-			MasterEncryptionKey = EncryptionKeys.Find(Other.MasterEncryptionKey->Guid);
+			SetPrincipalEncryptionKey(GetEncryptionKeys().Find(Other.GetPrincipalEncryptionKey()->Guid));
 		}
 		else
 		{
-			MasterEncryptionKey = nullptr;
+			SetPrincipalEncryptionKey(nullptr);
 		}
 
 		return *this;
@@ -70,26 +74,44 @@ struct FKeyChain
 
 	FKeyChain& operator=(FKeyChain&& Other)
 	{
-		SigningKey = Other.SigningKey;
-		EncryptionKeys = MoveTemp(Other.EncryptionKeys);
+		SetSigningKey(Other.GetSigningKey());
+		SetEncryptionKeys(MoveTemp(Other.GetEncryptionKeys()));
 		
-		if (Other.MasterEncryptionKey)
+		if (Other.GetPrincipalEncryptionKey())
 		{
-			MasterEncryptionKey = EncryptionKeys.Find(Other.MasterEncryptionKey->Guid);
+			SetPrincipalEncryptionKey(GetEncryptionKeys().Find(Other.GetPrincipalEncryptionKey()->Guid));
 		}
 		else
 		{
-			MasterEncryptionKey = nullptr;
+			SetPrincipalEncryptionKey(nullptr);
 		}
 
-		Other.SigningKey = InvalidRSAKeyHandle;
-		Other.MasterEncryptionKey = nullptr;
+		Other.SetSigningKey(InvalidRSAKeyHandle);
+		Other.SetPrincipalEncryptionKey(nullptr);
 
 		return *this;
 	}
 
+	PRAGMA_DISABLE_DEPRECATION_WARNINGS
+	FRSAKeyHandle GetSigningKey() const { return SigningKey; }
+	void SetSigningKey(FRSAKeyHandle key) { SigningKey = key; }
+
+	const FNamedAESKey* GetPrincipalEncryptionKey() const { return MasterEncryptionKey; }
+	void SetPrincipalEncryptionKey(const FNamedAESKey* key) { MasterEncryptionKey =key; }
+
+	const TMap<FGuid, FNamedAESKey>& GetEncryptionKeys() const { return EncryptionKeys; }
+	TMap<FGuid, FNamedAESKey>& GetEncryptionKeys() { return EncryptionKeys; }
+
+	void SetEncryptionKeys(const TMap<FGuid, FNamedAESKey>& keys) { EncryptionKeys = keys; }
+	PRAGMA_ENABLE_DEPRECATION_WARNINGS
+
+	UE_DEPRECATED(5.1, "Use Get/SetSigningKey instead")
 	FRSAKeyHandle SigningKey = InvalidRSAKeyHandle;
+
+	UE_DEPRECATED(5.1, "Use GetEncryptionKeys instead")
 	TMap<FGuid, FNamedAESKey> EncryptionKeys;
+	
+	UE_DEPRECATED(5.1, "Use Get/SetPrincipalEncryptionKey instead")
 	const FNamedAESKey* MasterEncryptionKey = nullptr;
 };
 
@@ -146,7 +168,7 @@ namespace KeyChainUtilities
 						NewKey.Name = TEXT("Default");
 						NewKey.Guid = FGuid();
 						FMemory::Memcpy(NewKey.Key.Key, &Key[0], sizeof(FAES::FAESKey::Key));
-						OutCryptoSettings.EncryptionKeys.Add(NewKey.Guid, NewKey);
+						OutCryptoSettings.GetEncryptionKeys().Add(NewKey.Guid, NewKey);
 					}
 				}
 			}
@@ -154,7 +176,7 @@ namespace KeyChainUtilities
 			const TSharedPtr<FJsonObject>* SigningKey = nullptr;
 			if (RootObject->TryGetObjectField(TEXT("SigningKey"), SigningKey))
 			{
-				OutCryptoSettings.SigningKey = ParseRSAKeyFromJson(*SigningKey);
+				OutCryptoSettings.SetSigningKey(ParseRSAKeyFromJson(*SigningKey));
 			}
 
 			const TArray<TSharedPtr<FJsonValue>>* SecondaryEncryptionKeyArray = nullptr;
@@ -173,25 +195,25 @@ namespace KeyChainUtilities
 					check(Key.Num() == sizeof(FAES::FAESKey::Key));
 					FMemory::Memcpy(NewKey.Key.Key, &Key[0], sizeof(FAES::FAESKey::Key));
 
-					check(!OutCryptoSettings.EncryptionKeys.Contains(NewKey.Guid) || OutCryptoSettings.EncryptionKeys[NewKey.Guid].Key == NewKey.Key);
-					OutCryptoSettings.EncryptionKeys.Add(NewKey.Guid, NewKey);
+					check(!OutCryptoSettings.GetEncryptionKeys().Contains(NewKey.Guid) || OutCryptoSettings.GetEncryptionKeys()[NewKey.Guid].Key == NewKey.Key);
+					OutCryptoSettings.GetEncryptionKeys().Add(NewKey.Guid, NewKey);
 				}
 			}
 		}
 		delete File;
 		FGuid EncryptionKeyOverrideGuid;
-		OutCryptoSettings.MasterEncryptionKey = OutCryptoSettings.EncryptionKeys.Find(EncryptionKeyOverrideGuid);
+		OutCryptoSettings.SetPrincipalEncryptionKey(OutCryptoSettings.GetEncryptionKeys().Find(EncryptionKeyOverrideGuid));
 	}
 
 	static void ApplyEncryptionKeys(const FKeyChain& KeyChain)
 	{
-		if (KeyChain.EncryptionKeys.Contains(FGuid()))
+		if (KeyChain.GetEncryptionKeys().Contains(FGuid()))
 		{
-			FAES::FAESKey DefaultKey = KeyChain.EncryptionKeys[FGuid()].Key;
+			FAES::FAESKey DefaultKey = KeyChain.GetEncryptionKeys()[FGuid()].Key;
 			FCoreDelegates::GetPakEncryptionKeyDelegate().BindLambda([DefaultKey](uint8 OutKey[32]) { FMemory::Memcpy(OutKey, DefaultKey.Key, sizeof(DefaultKey.Key)); });
 		}
 
-		for (const TMap<FGuid, FNamedAESKey>::ElementType& Key : KeyChain.EncryptionKeys)
+		for (const TMap<FGuid, FNamedAESKey>::ElementType& Key : KeyChain.GetEncryptionKeys())
 		{
 			if (Key.Key.IsValid())
 			{
