UnrealPak: Resolves an extremely rare race condition that happens when Flush() reaches the CompressionThread.Wait() before the CompressionThread has an opportunity to run and check for the flush condition.

diff --git a/Engine/Source/Developer/PakFileUtilities/Private/PakFileUtilities.cpp b/Engine/Source/Developer/PakFileUtilities/Private/PakFileUtilities.cpp
--- a/Engine/Source/Developer/PakFileUtilities/Private/PakFileUtilities.cpp
+++ b/Engine/Source/Developer/PakFileUtilities/Private/PakFileUtilities.cpp
@@ -2474,6 +2474,9 @@ bool FPakWriterContext::Flush()
 	check(!bFlushed);
 	bFlushed = true;
 
+	CompressionQueueEntryAddedEvent->Trigger();
+	WriteQueueEntryAddedEvent->Trigger();
+
 	const uint32 RequiredPatchPadding = CmdLineParameters.PatchFilePadAlign;
 
 	for (TUniquePtr<FOutputPakFile>& OutputPakFile : OutputPakFiles)
@@ -2646,6 +2649,7 @@ bool FPakWriterContext::Flush()
 
 	WriterThread.Wait();
 	CompressionThread.Wait();
+	UE_LOG(LogPakFile, Display, TEXT("Writer and Compression Threads exited."));
 
 	if (TotalFilesWithPoorForcedCompression > 0)
 	{
