UnrealPak: Refactor compression code to avoid busy waiting from within a task

diff --git a/Engine/Source/Developer/PakFileUtilities/Private/PakFileUtilities.cpp b/Engine/Source/Developer/PakFileUtilities/Private/PakFileUtilities.cpp
--- a/Engine/Source/Developer/PakFileUtilities/Private/PakFileUtilities.cpp
+++ b/Engine/Source/Developer/PakFileUtilities/Private/PakFileUtilities.cpp
@@ -82,7 +82,9 @@ class FBlockCompressTask : public FNonAbandonableTask
 public:
 	friend class FAsyncTask<FBlockCompressTask>;
 	friend class FMemoryCompressor;
-	FBlockCompressTask(void* InUncompressedBuffer, int32 InUncompressedSize, FName InFormat, int32 InBlockSize) :
+	FBlockCompressTask(void* InUncompressedBuffer, int32 InUncompressedSize, FName InFormat, int32 InBlockSize, std::atomic<int32>& InRemainingTasksCounter, FGraphEventRef InAllTasksFinishedEvent) :
+		RemainingTasksCounter(InRemainingTasksCounter),
+		AllTasksFinishedEvent(InAllTasksFinishedEvent),
 		UncompressedBuffer(InUncompressedBuffer),
 		UncompressedSize(InUncompressedSize),
 		Format(InFormat),
@@ -102,15 +104,25 @@ public:
 	/** Do compress */
 	void DoWork()
 	{
+#if DETAILED_UNREALPAK_TIMING
+		FUnrealPakScopeCycleCounter Scope(GCompressionTime);
+#endif
 		// Compress memory block.
 		// Actual size will be stored to CompressedSize.
 		Result = FCompression::CompressMemory(Format, CompressedBuffer, CompressedSize, UncompressedBuffer, UncompressedSize, COMPRESS_ForPackaging);
+		if (--RemainingTasksCounter == 0)
+		{
+			AllTasksFinishedEvent->DispatchSubsequents();
+		}
 	}
 
 	
 	FORCEINLINE TStatId GetStatId() const { RETURN_QUICK_DECLARE_CYCLE_STAT(ExampleAsyncTask, STATGROUP_ThreadPoolAsyncTasks); }
 		
 private:
+	std::atomic<int32>& RemainingTasksCounter;
+	FGraphEventRef AllTasksFinishedEvent;
+
 	// Source buffer
 	void* UncompressedBuffer;
 	int32 UncompressedSize;
@@ -133,7 +145,7 @@ class FMemoryCompressor
 {
 public:
 	/** Divide into blocks and start compress asynchronously */
-	FMemoryCompressor(uint8* UncompressedBuffer, int64 UncompressedSize, FName Format, int32 CompressionBlockSize) :
+	FMemoryCompressor(uint8* UncompressedBuffer, int64 UncompressedSize, FName Format, int32 CompressionBlockSize, FGraphEventRef CompressionFinishedEvent) :
 		Index(0)
 	{
 		// Divide into blocks and start compression async tasks.
@@ -142,36 +154,50 @@ public:
 		while (UncompressedSize)
 		{
 			int32 BlockSize = (int32)FMath::Min<int64>(UncompressedSize, CompressionBlockSize);
-			auto* AsyncTask = new FAsyncTask<FBlockCompressTask>(UncompressedBuffer + UncompressedBytes, BlockSize, Format, BlockSize);
-			AsyncTask->StartBackgroundTask();
+			auto* AsyncTask = new FAsyncTask<FBlockCompressTask>(UncompressedBuffer + UncompressedBytes, BlockSize, Format, BlockSize, RemainingTasksCounter, CompressionFinishedEvent);
 			BlockCompressAsyncTasks.Add(AsyncTask);
 			UncompressedSize -= BlockSize;
 			UncompressedBytes += BlockSize;
 		}
+
+		if (BlockCompressAsyncTasks.IsEmpty())
+		{
+			CompressionFinishedEvent->DispatchSubsequents();
+		}
+		else
+		{
+			RemainingTasksCounter = BlockCompressAsyncTasks.Num();
+
+			for (auto* AsyncTask : BlockCompressAsyncTasks)
+			{
+				AsyncTask->StartBackgroundTask();
+			}
+		}
 	}
 
 	~FMemoryCompressor()
 	{
 		for (auto* AsyncTask : BlockCompressAsyncTasks)
 		{
-			if (!AsyncTask->Cancel())
-			{
-				AsyncTask->EnsureCompletion();
-			}
+			check(AsyncTask->IsDone());
 			delete AsyncTask;
 		}
 	}
 
 
 	/** Fetch compressed result. Returns true and store CompressedSize if succeeded */
-	bool CompressMemory(FName Format, void* CompressedBuffer, int32& CompressedSize, const void* UncompressedBuffer, int32 UncompressedSize)
+	bool CopyNextCompressedBuffer(FName Format, void* CompressedBuffer, int32& CompressedSize, const void* UncompressedBuffer, int32 UncompressedSize)
 	{
 		// Fetch compressed result from task.
 		// We assume this is called only once, same order, same parameters for
 		// each task.
 		auto* AsyncTask = BlockCompressAsyncTasks[Index++];
-		AsyncTask->EnsureCompletion();
-
+		while (!AsyncTask->IsDone())
+		{
+			// Compression is done but we also need to wait for the async task to be marked as completed before we call GetTask() below
+			// DON'T call EnsureCompletion here, we're already running in a task so we don't want to pull in other tasks while waiting
+			FPlatformProcess::Sleep(0);
+		}
 		FBlockCompressTask& Task = AsyncTask->GetTask();
 		check(Task.Format == Format);
 		check(Task.UncompressedBuffer == UncompressedBuffer);
@@ -189,6 +215,7 @@ public:
 
 private:
 	TArray<FAsyncTask<FBlockCompressTask>*> BlockCompressAsyncTasks;
+	std::atomic<int32> RemainingTasksCounter = 0;
 
 	// Fetched task index
 	int32 Index;
@@ -565,7 +592,8 @@ struct FCompressedFileBuffer
 
 	bool ReadUncompressedFileToWorkingBuffer(const FPakInputPair& InFile);
 	
-	bool CompressFileToWorkingBuffer(const FPakInputPair& InFile, FName CompressionMethod, const int32 CompressionBlockSize);
+	TUniquePtr<FMemoryCompressor> BeginCompressFileToWorkingBuffer(const FPakInputPair& InFile, FName CompressionMethod, const int32 CompressionBlockSize, FGraphEventRef EndCompressionBarrier);
+	bool EndCompressFileToWorkingBuffer(const FPakInputPair& InFile, FName CompressionMethod, const int32 CompressionBlockSize, FMemoryCompressor& MemoryCompressor);
 
 	void SerializeDDCData(FArchive &Ar)
 	{
@@ -592,6 +620,7 @@ struct FCompressedFileBuffer
 
 	FString GetDDCKeyString(const uint8* UncompressedFile, const int64& UncompressedFileSize, FName CompressionFormat, const int64& BlockSize);
 
+	TArray<uint8>		UncompressedBuffer;
 	int64				OriginalSize;
 	int64				TotalCompressedSize;
 	int32				FileCompressionBlockSize;
@@ -718,25 +747,24 @@ bool FCompressedFileBuffer::ReadUncompressedFileToWorkingBuffer(const FPakInputP
 	return true;
 }
 
-bool FCompressedFileBuffer::CompressFileToWorkingBuffer(const FPakInputPair& InFile, FName CompressionMethod, const int32 CompressionBlockSize)
+TUniquePtr<FMemoryCompressor> FCompressedFileBuffer::BeginCompressFileToWorkingBuffer(const FPakInputPair& InFile, FName CompressionMethod, const int32 CompressionBlockSize, FGraphEventRef EndCompressionBarrier)
 {
 	TUniquePtr<FArchive> FileHandle(IFileManager::Get().CreateFileReader(*InFile.Source));
 	if(!FileHandle)
 	{
 		TotalCompressedSize = 0;
-		return false;
+		EndCompressionBarrier->DispatchSubsequents();
+		return nullptr;
 	}
 
 	Reinitialize(FileHandle.Get(), CompressionMethod, CompressionBlockSize);
-	const int64 FileSize = OriginalSize;
-	const int64 PaddedEncryptedFileSize = Align(FileSize,FAES::AESBlockSize);
-	check(PaddedEncryptedFileSize >= FileSize);
+	const int64 PaddedEncryptedFileSize = Align(OriginalSize, FAES::AESBlockSize);
+	check(PaddedEncryptedFileSize >= OriginalSize);
 
-	TArray<uint8, TInlineAllocator64<4096>> UncompressedBuffer;
 	UncompressedBuffer.SetNumUninitialized(PaddedEncryptedFileSize);
 
 	// Load to buffer
-	FileHandle->Serialize(UncompressedBuffer.GetData(), FileSize);
+	FileHandle->Serialize(UncompressedBuffer.GetData(), OriginalSize);
 
 #if USE_DDC_FOR_COMPRESSED_FILES
 	const bool bShouldUseDDC = true; // && (FileSize > 20 * 1024 ? true : false);
@@ -760,21 +788,27 @@ bool FCompressedFileBuffer::CompressFileToWorkingBuffer(const FPakInputPair& InF
 			{
 				FMemoryReader Ar(GetData, true);
 				SerializeDDCData(Ar);
-				return true;
+				UncompressedBuffer.Empty();
+				EndCompressionBarrier->DispatchSubsequents();
+				return nullptr;
 			}
 		}
 	}
 #endif
 
-	{
-#if DETAILED_UNREALPAK_TIMING
-		FUnrealPakScopeCycleCounter Scope(GCompressionTime);
-#endif
-		// Start parallel compress
-		FMemoryCompressor MemoryCompressor(UncompressedBuffer.GetData(), FileSize, CompressionMethod, CompressionBlockSize);
+	// Start parallel compress
+	return MakeUnique<FMemoryCompressor>(UncompressedBuffer.GetData(), OriginalSize, CompressionMethod, CompressionBlockSize, EndCompressionBarrier);
+}
 
+bool FCompressedFileBuffer::EndCompressFileToWorkingBuffer(const FPakInputPair& InFile, FName CompressionMethod, const int32 CompressionBlockSize, FMemoryCompressor& MemoryCompressor)
+{
+	ON_SCOPE_EXIT
+	{
+		UncompressedBuffer.Empty();
+	};
+	{
 		// Build buffers for working
-		int64 UncompressedSize = FileSize;
+		int64 UncompressedSize = OriginalSize;
 
 		// CompressMemoryBound truncates its size argument to 32bits, so we can not use (possibly > 32-bit) FileSize directly to calculate required buffer space
 		int32 MaxCompressedBufferSize = Align(FCompression::CompressMemoryBound(CompressionMethod, CompressionBlockSize, COMPRESS_NoFlags), FAES::AESBlockSize);
@@ -791,7 +825,7 @@ bool FCompressedFileBuffer::CompressFileToWorkingBuffer(const FPakInputPair& InF
 			int32 CompressedBlockSize = FMath::Max<int32>(MaxCompressedBufferSize, MaxCompressedBlockSize);
 			FileCompressionBlockSize = FMath::Max<uint32>(BlockSize, FileCompressionBlockSize);
 			EnsureBufferSpace(Align(TotalCompressedSize + CompressedBlockSize, FAES::AESBlockSize));
-			if (!MemoryCompressor.CompressMemory(CompressionMethod, CompressedBuffer.Get() + TotalCompressedSize, CompressedBlockSize, UncompressedBuffer.GetData() + UncompressedBytes, BlockSize))
+			if (!MemoryCompressor.CopyNextCompressedBuffer(CompressionMethod, CompressedBuffer.Get() + TotalCompressedSize, CompressedBlockSize, UncompressedBuffer.GetData() + UncompressedBytes, BlockSize))
 			{
 				return false;
 			}
@@ -1836,8 +1870,11 @@ private:
 		FName CompressionMethod = NAME_None;
 		int64 RealFileSize = 0;
 
-		FGraphEventRef CompressionBarrier;
-		FGraphEventRef CompressionTask;
+		TUniquePtr<FMemoryCompressor> MemoryCompressor;
+		FGraphEventRef BeginCompressionBarrier;
+		FGraphEventRef BeginCompressionTask;
+		FGraphEventRef EndCompressionBarrier;
+		FGraphEventRef EndCompressionTask;
 	};
 
 	struct FOutputPakFile
@@ -1870,7 +1907,8 @@ private:
 		int32 RetiredEntriesCount = 0;
 	};
 
-	void Compress(FOutputPakFileEntry* Entry);
+	void BeginCompress(FOutputPakFileEntry* Entry);
+	void EndCompress(FOutputPakFileEntry* Entry);
 	void Retire(FOutputPakFileEntry* Entry);
 	void CompressionThreadFunc();
 	void WriterThreadFunc();
@@ -2053,11 +2091,16 @@ bool FPakWriterContext::AddPakFile(const TCHAR* Filename, const TArray<FPakInput
 				OutputEntry.UExpFileInPair = &OutputPakFile->Entries[FileIndex + 1];
 			}
 		}
-		OutputEntry.CompressionBarrier = FGraphEvent::CreateGraphEvent();
-		OutputEntry.CompressionTask = FFunctionGraphTask::CreateAndDispatchWhenReady([this, &OutputEntry]()
+		OutputEntry.BeginCompressionBarrier = FGraphEvent::CreateGraphEvent();
+		OutputEntry.BeginCompressionTask = FFunctionGraphTask::CreateAndDispatchWhenReady([this, &OutputEntry]()
 			{
-				Compress(&OutputEntry);
-			}, TStatId(), OutputEntry.CompressionBarrier, ENamedThreads::AnyHiPriThreadHiPriTask);
+				BeginCompress(&OutputEntry);
+			}, TStatId(), OutputEntry.BeginCompressionBarrier, ENamedThreads::AnyHiPriThreadHiPriTask);
+		OutputEntry.EndCompressionBarrier = FGraphEvent::CreateGraphEvent();
+		OutputEntry.EndCompressionTask = FFunctionGraphTask::CreateAndDispatchWhenReady([this, &OutputEntry]()
+			{
+				EndCompress(&OutputEntry);
+			}, TStatId(), OutputEntry.EndCompressionBarrier, ENamedThreads::AnyHiPriThreadHiPriTask);
 		CompressionQueue.Enqueue(&OutputEntry);
 		WriteQueue.Enqueue(&OutputEntry);
 	}
@@ -2079,9 +2122,9 @@ bool FPakWriterContext::AddPakFile(const TCHAR* Filename, const TArray<FPakInput
 	return true;
 }
 
-void FPakWriterContext::Compress(FOutputPakFileEntry* Entry)
+void FPakWriterContext::BeginCompress(FOutputPakFileEntry* Entry)
 {
-	TRACE_CPUPROFILER_EVENT_SCOPE(Compress);
+	TRACE_CPUPROFILER_EVENT_SCOPE(BeginCompress);
 
 	// don't try to compress tiny files
 	// even if they do compress, it is a bad use of decoder time
@@ -2089,6 +2132,7 @@ void FPakWriterContext::Compress(FOutputPakFileEntry* Entry)
 	{
 		Entry->CompressedFileBuffer.ReadUncompressedFileToWorkingBuffer(Entry->InputPair);
 		Entry->bSomeCompressionSucceeded = true; // Only used for logging purposes
+		Entry->EndCompressionBarrier->DispatchSubsequents();
 		return;
 	}
 
@@ -2096,23 +2140,30 @@ void FPakWriterContext::Compress(FOutputPakFileEntry* Entry)
 	// if Oodle refused to compress it was not an error, it was a choice because Oodle didn't think
 	//  compressing that file was worth the decode time
 	//	we do NOT want Zlib to then get enabled for that file!
-	for (int32 MethodIndex = 0; MethodIndex < 1; MethodIndex++)
+	const int32 MethodIndex = 0;
 	{
 		Entry->CompressionMethod = CmdLineParameters.CompressionFormats[MethodIndex];
 
 		// because compression is a plugin, certain files need to be loadable out of pak files before plugins are loadable
 		// (like .uplugin files). for these, we enforce a non-plugin compression - zlib
 		//	note that those file types are also excluded from iostore, so still go through this pak system
-		bool bForceCompressionFormat = false;
 		if (NoPluginCompressionExtensions.Find(FPaths::GetExtension(Entry->InputPair.Source)) != nullptr ||
 			NoPluginCompressionFileNames.Find(FPaths::GetCleanFilename(Entry->InputPair.Source)) != nullptr)
 		{
 			Entry->CompressionMethod = NAME_Zlib;
-			bForceCompressionFormat = true;
 		}
+	}
+	// attempt to compress the data
+	Entry->MemoryCompressor = Entry->CompressedFileBuffer.BeginCompressFileToWorkingBuffer(Entry->InputPair, Entry->CompressionMethod, CmdLineParameters.CompressionBlockSize, Entry->EndCompressionBarrier);
+}
+
+void FPakWriterContext::EndCompress(FOutputPakFileEntry* Entry)
+{
+	TRACE_CPUPROFILER_EVENT_SCOPE(EndCompress);
 
-		// attempt to compress the data
-		if (Entry->CompressedFileBuffer.CompressFileToWorkingBuffer(Entry->InputPair, Entry->CompressionMethod, CmdLineParameters.CompressionBlockSize))
+	const int32 MethodIndex = 0;
+	{
+		if (Entry->MemoryCompressor && Entry->CompressedFileBuffer.EndCompressFileToWorkingBuffer(Entry->InputPair, Entry->CompressionMethod, CmdLineParameters.CompressionBlockSize, *Entry->MemoryCompressor.Get()))
 		{
 			// for modern compressors we don't want any funny heuristics turning compression on/off
 			// let the compressor decide; it has an introspective measure of whether compression is worth doing or not
@@ -2158,13 +2209,8 @@ void FPakWriterContext::Compress(FOutputPakFileEntry* Entry)
 				}
 			}
 		}
-
-		// if we successfully compressed it, or we only wanted a single format, then we are done!
-		if (Entry->bSomeCompressionSucceeded || bForceCompressionFormat)
-		{
-			break;
-		}
 	}
+	Entry->MemoryCompressor.Reset();
 
 	// If no compression was able to make it small enough, or compress at all, don't compress it
 	if (!Entry->bSomeCompressionSucceeded)
@@ -2218,7 +2264,7 @@ void FPakWriterContext::CompressionThreadFunc()
 				}
 			}
 			ScheduledFileSize += Entry->OriginalFileSize;
-			Entry->CompressionBarrier->DispatchSubsequents();
+			Entry->BeginCompressionBarrier->DispatchSubsequents();
 		}
 		else
 		{
@@ -2260,10 +2306,10 @@ void FPakWriterContext::WriterThreadFunc()
 		int64 NewEntryOffset = PakFileHandle->Tell();
 		FPakEntryPair& NewEntry = OutputEntry->Entry;
 
-		if (!OutputEntry->CompressionTask->IsComplete())
+		if (!OutputEntry->EndCompressionTask->IsComplete())
 		{
 			TRACE_CPUPROFILER_EVENT_SCOPE(WaitForCompression);
-			OutputEntry->CompressionTask->Wait();
+			OutputEntry->EndCompressionTask->Wait();
 		}
 		
 		const FName& CompressionMethod = OutputEntry->CompressionMethod;
