Unshelved from pending changelist '21985878':

- Fixed an assert with -makebinaryconfig

diff --git a/Engine/Source/Developer/PakFileUtilities/Private/PakFileUtilities.cpp b/Engine/Source/Developer/PakFileUtilities/Private/PakFileUtilities.cpp
--- a/Engine/Source/Developer/PakFileUtilities/Private/PakFileUtilities.cpp
+++ b/Engine/Source/Developer/PakFileUtilities/Private/PakFileUtilities.cpp
@@ -5040,7 +5040,7 @@ bool MakeBinaryConfig(const TCHAR* CmdLine)
 
 	FConfigCacheIni Config(EConfigCacheType::Temporary);
 	FConfigContext Context = FConfigContext::ReadIntoConfigSystem(&Config, TEXT(""));
-	Context.ProjectConfigDir = FPaths::Combine(ProjectDir, TEXT("Config"));
+	Context.ProjectConfigDir = FPaths::Combine(ProjectDir, TEXT("Config/"));
 	Config.InitializeKnownConfigFiles(Context);
 
 	// removing for now, because this causes issues with some plugins not getting ini files merged in
