UnrealPak: Fixed issue with files larger than 2GB

diff --git a/Engine/Source/Developer/PakFileUtilities/Private/PakFileUtilities.cpp b/Engine/Source/Developer/PakFileUtilities/Private/PakFileUtilities.cpp
--- a/Engine/Source/Developer/PakFileUtilities/Private/PakFileUtilities.cpp
+++ b/Engine/Source/Developer/PakFileUtilities/Private/PakFileUtilities.cpp
@@ -620,7 +620,7 @@ struct FCompressedFileBuffer
 
 	FString GetDDCKeyString(const uint8* UncompressedFile, const int64& UncompressedFileSize, FName CompressionFormat, const int64& BlockSize);
 
-	TArray<uint8>		UncompressedBuffer;
+	TArray64<uint8>		UncompressedBuffer;
 	int64				OriginalSize;
 	int64				TotalCompressedSize;
 	int32				FileCompressionBlockSize;
