Fix a bug where UnrealPak would produce inconsistent binary output due to occasionally skipping compression on files that should've been compressed.

#rnx

### Problem
- The bug was originally introduced in CL 21791765
- Each file to be compressed calls FPakWriterContext::BeginCompress which will create a FMemoryCompressor and assign it to the files FOutputPakFileEntry. The problem is that the work done by the FMemoryCompressor can potentially complete and signal EndCompressionBarrier before the FMemoryCompresseor has been assigned. This can be due to the work itself being very small, or that closing a file handle when FCompressedFileBuffer::BeginCompressFileToWorkingBuffer returns (which creates the FMemoryCompressor) stalls against system resources. This potentially allows FPakWriterContext::EndCompres to execute concurrently and incorrectly believe that the file is not being compressed as FOutputPakFileEntry::MemoryCompressor is still nullptr.
- Not only does this mean that some files end up being paked uncompressed when they should be compressed, the results are not deterministic which means we can end up with different binary output when paking the exact same data with the exact same settings.

### Solution
- The constructor of FMemoryCompressor no longer adds its work to the task graph system. Instead the tasks are only added after the compressor has been assigned to FOutputPakFileEntry.
- Replaced auto in ranged for loops with the type, as per recent coding standards discussions.

diff --git a/Engine/Source/Developer/PakFileUtilities/Private/PakFileUtilities.cpp b/Engine/Source/Developer/PakFileUtilities/Private/PakFileUtilities.cpp
--- a/Engine/Source/Developer/PakFileUtilities/Private/PakFileUtilities.cpp
+++ b/Engine/Source/Developer/PakFileUtilities/Private/PakFileUtilities.cpp
@@ -145,8 +145,9 @@ class FMemoryCompressor
 {
 public:
 	/** Divide into blocks and start compress asynchronously */
-	FMemoryCompressor(uint8* UncompressedBuffer, int64 UncompressedSize, FName Format, int32 CompressionBlockSize, FGraphEventRef CompressionFinishedEvent) :
-		Index(0)
+	FMemoryCompressor(uint8* UncompressedBuffer, int64 UncompressedSize, FName Format, int32 CompressionBlockSize, FGraphEventRef InCompressionFinishedEvent) 
+		: CompressionFinishedEvent(InCompressionFinishedEvent)
+		, Index(0)
 	{
 		// Divide into blocks and start compression async tasks.
 		// These blocks must be as same as followed CompressMemory callings.
@@ -158,8 +159,20 @@ public:
 			BlockCompressAsyncTasks.Add(AsyncTask);
 			UncompressedSize -= BlockSize;
 			UncompressedBytes += BlockSize;
+		}	
+	}
+
+	~FMemoryCompressor()
+	{
+		for (FAsyncTask<FBlockCompressTask>* AsyncTask : BlockCompressAsyncTasks)
+		{
+			check(AsyncTask->IsDone());
+			delete AsyncTask;
 		}
+	}
 
+	void StartWork()
+	{
 		if (BlockCompressAsyncTasks.IsEmpty())
 		{
 			CompressionFinishedEvent->DispatchSubsequents();
@@ -168,30 +181,20 @@ public:
 		{
 			RemainingTasksCounter = BlockCompressAsyncTasks.Num();
 
-			for (auto* AsyncTask : BlockCompressAsyncTasks)
+			for (FAsyncTask<FBlockCompressTask>* AsyncTask : BlockCompressAsyncTasks)
 			{
 				AsyncTask->StartBackgroundTask();
 			}
 		}
 	}
 
-	~FMemoryCompressor()
-	{
-		for (auto* AsyncTask : BlockCompressAsyncTasks)
-		{
-			check(AsyncTask->IsDone());
-			delete AsyncTask;
-		}
-	}
-
-
 	/** Fetch compressed result. Returns true and store CompressedSize if succeeded */
 	bool CopyNextCompressedBuffer(FName Format, void* CompressedBuffer, int32& CompressedSize, const void* UncompressedBuffer, int32 UncompressedSize)
 	{
 		// Fetch compressed result from task.
 		// We assume this is called only once, same order, same parameters for
 		// each task.
-		auto* AsyncTask = BlockCompressAsyncTasks[Index++];
+		FAsyncTask<FBlockCompressTask>* AsyncTask = BlockCompressAsyncTasks[Index++];
 		while (!AsyncTask->IsDone())
 		{
 			// Compression is done but we also need to wait for the async task to be marked as completed before we call GetTask() below
@@ -214,6 +217,8 @@ public:
 	}
 
 private:
+	FGraphEventRef CompressionFinishedEvent;
+
 	TArray<FAsyncTask<FBlockCompressTask>*> BlockCompressAsyncTasks;
 	std::atomic<int32> RemainingTasksCounter = 0;
 
@@ -2155,6 +2160,10 @@ void FPakWriterContext::BeginCompress(FOutputPakFileEntry* Entry)
 	}
 	// attempt to compress the data
 	Entry->MemoryCompressor = Entry->CompressedFileBuffer.BeginCompressFileToWorkingBuffer(Entry->InputPair, Entry->CompressionMethod, CmdLineParameters.CompressionBlockSize, Entry->EndCompressionBarrier);
+	if (Entry->MemoryCompressor != nullptr)
+	{
+		Entry->MemoryCompressor->StartWork();
+	}
 }
 
 void FPakWriterContext::EndCompress(FOutputPakFileEntry* Entry)
