UnrealPak now supports rehydrating virtualized files as they are placed into pak files

#rnx

- When parsing the paklist, any file with a -rehydrate flag will be considered for rehydration

- Add the virtualization module to the PakFileUtilities module
- Add the source control module to UnrealPak (the virtualization module should be taking care of this)
- While parsing the files to be included in the pak file, we will record if any of the files require rehydration, if so this will be noted in FPakCommandLineParameters::bRequiresRehydration and used to initialize the virtualization system.
-- We only need to initialize the system once, even if we detect that multiple pak files have files that need rehydration.
-- If no pak file needs rehydration we do not initialize the system and the virtualization module is never loaded.
- FOutputPakFileEntry::CompressedFileBuffer was made private with accessors to make the refactor easier.
- Now when we start to compress a file, we always load it entirely into memory (along with any potential padding needed), if we then detect that the file doesn't require padding then we just use the in memory buffer as the output and if we do need compression we compress the buffer.
-- This means that there is only one place we load the file from disk, meaing only one place we need to insert the rehydration code to. This load is where the rehydration occurs.
-- In the previous code we would load the file into memory and then retain this copy during the compression pass, so no additional memory should be used.
- At the moment when we get the buffer back fromt he rehydration pass it will be in the form of a FSharedBuffer and so for now we need to memcpy to FCompressedFileBuffer::UncompressedBuffer, which is a waste of cpu cycles.
-- In a future change we should change UncompressedBuffer from TArray to FSharedBuffer to avoid this.

diff --git a/Engine/Source/Developer/PakFileUtilities/PakFileUtilities.Build.cs b/Engine/Source/Developer/PakFileUtilities/PakFileUtilities.Build.cs
--- a/Engine/Source/Developer/PakFileUtilities/PakFileUtilities.Build.cs
+++ b/Engine/Source/Developer/PakFileUtilities/PakFileUtilities.Build.cs
@@ -22,7 +22,10 @@ public class PakFileUtilities : ModuleRules
 
 		if (Target.bBuildWithEditorOnlyData)
 		{
-			DynamicallyLoadedModuleNames.Add("DerivedDataCache");
+			DynamicallyLoadedModuleNames.AddRange(new string[] {
+				"DerivedDataCache",
+				"Virtualization",
+			});
 		}
 	}
 }
diff --git a/Engine/Source/Developer/PakFileUtilities/Private/PakFileUtilities.cpp b/Engine/Source/Developer/PakFileUtilities/Private/PakFileUtilities.cpp
--- a/Engine/Source/Developer/PakFileUtilities/Private/PakFileUtilities.cpp
+++ b/Engine/Source/Developer/PakFileUtilities/Private/PakFileUtilities.cpp
@@ -33,6 +33,7 @@
 #include "Containers/SpscQueue.h"
 #include "Async/Async.h"
 #include "Async/Future.h"
+#include "Virtualization/VirtualizationSystem.h"
 
 IMPLEMENT_MODULE(FDefaultModuleImpl, PakFileUtilities);
 
@@ -471,6 +472,7 @@ struct FPakCommandLineParameters
 		, bAlignFilesLargerThanBlock(false)
 		, bForceCompress(false)
 		, bFileRegions(false)
+		, bRequiresRehydration(false)
 	{
 	}
 
@@ -495,6 +497,8 @@ struct FPakCommandLineParameters
 
 	bool bForceCompress; // Force all files that request compression to be compressed, even if that results in a larger file size DEPRECATED
 	bool bFileRegions; // Enables the processing and output of cook file region metadata, used during packaging on some platforms.
+
+	bool bRequiresRehydration; // One or more files to be added to the pak file are virtualized and require rehydration.
 };
 
 struct FPakInputPair
@@ -506,6 +510,7 @@ struct FPakInputPair
 	bool bNeedEncryption;
 	bool bIsDeleteRecord;	// This is used for patch PAKs when a file is deleted from one patch to the next
 	bool bIsInPrimaryOrder;
+	bool bNeedRehydration;
 
 	FPakInputPair()
 		: SuggestedOrder(MAX_uint64)
@@ -513,6 +518,7 @@ struct FPakInputPair
 		, bNeedEncryption(false)
 		, bIsDeleteRecord(false)
 		, bIsInPrimaryOrder(false)
+		, bNeedRehydration(false)
 	{}
 
 	FPakInputPair(const FString& InSource, const FString& InDest)
@@ -521,6 +527,7 @@ struct FPakInputPair
 		, bNeedsCompression(false)
 		, bNeedEncryption(false)
 		, bIsDeleteRecord(false)
+		, bNeedRehydration(false)
 	{}
 
 	FORCEINLINE bool operator==(const FPakInputPair& Other) const
@@ -556,16 +563,17 @@ struct FCompressedFileBuffer
 {
 	FCompressedFileBuffer()
 		: OriginalSize(0)
-		,TotalCompressedSize(0)
-		,FileCompressionBlockSize(0)
-		,CompressedBufferSize(0)
+		, TotalCompressedSize(0)
+		, FileCompressionBlockSize(0)
+		, CompressedBufferSize(0)
+		, RehydrationCount(0)
+		, RehydrationBytes(0)
 	{
 
 	}
 
-	void Reinitialize(FArchive* File, FName CompressionMethod, int64 CompressionBlockSize)
+	void Reinitialize(FName CompressionMethod, int64 CompressionBlockSize)
 	{
-		OriginalSize = File->TotalSize();
 		TotalCompressedSize = 0;
 		FileCompressionBlockSize = 0;
 		FileCompressionMethod = CompressionMethod;
@@ -582,6 +590,8 @@ struct FCompressedFileBuffer
 		CompressedBuffer = nullptr;
 		CompressedBufferSize = 0; 
 		CompressedBlocks.Empty();
+		RehydrationCount = 0;
+		RehydrationBytes = 0;
 	}
 
 	void EnsureBufferSpace(int64 RequiredSpace)
@@ -595,8 +605,10 @@ struct FCompressedFileBuffer
 		}
 	}
 
-	bool ReadUncompressedFileToWorkingBuffer(const FPakInputPair& InFile);
-	
+	void ResetSource();
+	bool ReadSource(const FPakInputPair& InFile);
+	void SetSourceAsWorkingBuffer();
+
 	TUniquePtr<FMemoryCompressor> BeginCompressFileToWorkingBuffer(const FPakInputPair& InFile, FName CompressionMethod, const int32 CompressionBlockSize, FGraphEventRef EndCompressionBarrier);
 	bool EndCompressFileToWorkingBuffer(const FPakInputPair& InFile, FName CompressionMethod, const int32 CompressionBlockSize, FMemoryCompressor& MemoryCompressor);
 
@@ -625,7 +637,15 @@ struct FCompressedFileBuffer
 
 	FString GetDDCKeyString(const uint8* UncompressedFile, const int64& UncompressedFileSize, FName CompressionFormat, const int64& BlockSize);
 
+	/** Returns the size of the file, excluding any padding that might have been applied */
+	int64 GetFileSize() const
+	{
+		return OriginalSize;
+	}
+
 	TArray64<uint8>		UncompressedBuffer;
+
+
 	int64				OriginalSize;
 	int64				TotalCompressedSize;
 	int32				FileCompressionBlockSize;
@@ -633,6 +653,8 @@ struct FCompressedFileBuffer
 	TArray<FPakCompressedBlock>	CompressedBlocks;
 	int64				CompressedBufferSize;
 	TUniquePtr<uint8[]>		CompressedBuffer;
+	int32				RehydrationCount;
+	int64				RehydrationBytes;
 };
 
 template <class T>
@@ -735,41 +757,105 @@ FString FCompressedFileBuffer::GetDDCKeyString(const uint8* UncompressedFile, co
 	return FDerivedDataCacheInterface::BuildCacheKey(TEXT("PAKCOMPRESS_"), PAKCOMPRESS_DERIVEDDATA_VER, *KeyString);
 }
 
-bool FCompressedFileBuffer::ReadUncompressedFileToWorkingBuffer(const FPakInputPair& InFile)
+void FCompressedFileBuffer::ResetSource()
 {
-	TUniquePtr<FArchive> FileHandle(IFileManager::Get().CreateFileReader(*InFile.Source));
-	if (!FileHandle)
-	{
-		OriginalSize = -1;
-		return false;
-	}
-	OriginalSize = FileHandle->TotalSize();
-	const int64 PaddedEncryptedFileSize = Align(OriginalSize, FAES::AESBlockSize);
-	CompressedBuffer = MakeUnique<uint8[]>(PaddedEncryptedFileSize);
-	CompressedBufferSize = PaddedEncryptedFileSize;
-	// Load to buffer
-	FileHandle->Serialize(CompressedBuffer.Get(), OriginalSize);
-	return true;
+	UncompressedBuffer.Empty();
 }
 
-TUniquePtr<FMemoryCompressor> FCompressedFileBuffer::BeginCompressFileToWorkingBuffer(const FPakInputPair& InFile, FName CompressionMethod, const int32 CompressionBlockSize, FGraphEventRef EndCompressionBarrier)
+bool FCompressedFileBuffer::ReadSource(const FPakInputPair& InputInfo)
 {
-	TUniquePtr<FArchive> FileHandle(IFileManager::Get().CreateFileReader(*InFile.Source));
-	if(!FileHandle)
+	if (!InputInfo.bNeedRehydration)
 	{
-		TotalCompressedSize = 0;
-		EndCompressionBarrier->DispatchSubsequents();
-		return nullptr;
+		TUniquePtr<FArchive> FileHandle(IFileManager::Get().CreateFileReader(*InputInfo.Source));
+		if (!FileHandle)
+		{
+			OriginalSize = -1;
+			return false;
+		}
+
+		OriginalSize = FileHandle->TotalSize();
+		const int64 PaddedEncryptedFileSize = Align(OriginalSize, FAES::AESBlockSize);
+
+		UncompressedBuffer.SetNumUninitialized(PaddedEncryptedFileSize);
+
+		FileHandle->Serialize(UncompressedBuffer.GetData(), OriginalSize);
+
+		if (!FileHandle->IsError())
+		{
+			return true;
+		}
+		else
+		{
+			UncompressedBuffer.Empty();
+			OriginalSize = -1;
+
+			return false;
+		}
 	}
+	else
+	{
+		using namespace UE::Virtualization;
+		IVirtualizationSystem& System = IVirtualizationSystem::Get();
+
+		TArray<FString> PackagePath;
+		PackagePath.Add(InputInfo.Source);
 
-	Reinitialize(FileHandle.Get(), CompressionMethod, CompressionBlockSize);
-	const int64 PaddedEncryptedFileSize = Align(OriginalSize, FAES::AESBlockSize);
-	check(PaddedEncryptedFileSize >= OriginalSize);
+		TArray<FSharedBuffer> RehydratedPackage;
+		TArray<FRehydrationInfo> RehydrationInfo;
+		TArray<FText> Errors;
+
+		if (System.TryRehydratePackages(PackagePath, FAES::AESBlockSize, Errors, RehydratedPackage, &RehydrationInfo) == ERehydrationResult::Success)
+		{
+			OriginalSize = RehydrationInfo[0].RehydratedSize;
+			const int64 PaddedEncryptedFileSize = Align(OriginalSize, FAES::AESBlockSize);
 
-	UncompressedBuffer.SetNumUninitialized(PaddedEncryptedFileSize);
+			// Confirm padding worked
+			check(RehydratedPackage[0].GetSize() == PaddedEncryptedFileSize);
 
-	// Load to buffer
-	FileHandle->Serialize(UncompressedBuffer.GetData(), OriginalSize);
+			//TODO: Keep data in FSharedBuffer form, rather than TArray
+			UncompressedBuffer.SetNumUninitialized(RehydratedPackage[0].GetSize());
+			FMemory::Memcpy(UncompressedBuffer.GetData(), RehydratedPackage[0].GetData(), RehydratedPackage[0].GetSize());
+
+			RehydrationCount = RehydrationInfo[0].NumPayloadsRehydrated;
+			RehydrationBytes = RehydrationInfo[0].RehydratedSize - RehydrationInfo[0].OriginalSize;
+
+			if (RehydrationCount)
+			{
+				UE_LOG(LogPakFile, Verbose, TEXT("Rehydrated (+%lld bytes) %s"), RehydrationBytes, *PackagePath[0]);
+			}
+
+			return true;
+		}
+		else
+		{
+			UncompressedBuffer.Empty();
+			OriginalSize = -1;
+
+			return false;
+		}
+	}
+}
+
+void FCompressedFileBuffer::SetSourceAsWorkingBuffer()
+{
+	// TODO: Do not submit with this assert
+	check(!UncompressedBuffer.IsEmpty() || OriginalSize == 0);
+
+	// TODO: Better way to move the buffers ownership, but can't safely steal the memory in a
+	// TArray
+	CompressedBufferSize = UncompressedBuffer.Num();
+	CompressedBuffer = MakeUnique<uint8[]>(CompressedBufferSize);
+
+	FMemory::Memcpy(CompressedBuffer.Get(), UncompressedBuffer.GetData(), CompressedBufferSize);
+	
+	UncompressedBuffer.Empty();
+}
+
+TUniquePtr<FMemoryCompressor> FCompressedFileBuffer::BeginCompressFileToWorkingBuffer(const FPakInputPair& InFile, FName CompressionMethod, const int32 CompressionBlockSize, FGraphEventRef EndCompressionBarrier)
+{
+	check(!UncompressedBuffer.IsEmpty() || OriginalSize == 0);
+
+	Reinitialize(CompressionMethod, CompressionBlockSize);
 
 #if USE_DDC_FOR_COMPRESSED_FILES
 	const bool bShouldUseDDC = true; // && (FileSize > 20 * 1024 ? true : false);
@@ -807,10 +893,6 @@ TUniquePtr<FMemoryCompressor> FCompressedFileBuffer::BeginCompressFileToWorkingB
 
 bool FCompressedFileBuffer::EndCompressFileToWorkingBuffer(const FPakInputPair& InFile, FName CompressionMethod, const int32 CompressionBlockSize, FMemoryCompressor& MemoryCompressor)
 {
-	ON_SCOPE_EXIT
-	{
-		UncompressedBuffer.Empty();
-	};
 	{
 		// Build buffers for working
 		int64 UncompressedSize = OriginalSize;
@@ -883,6 +965,7 @@ bool PrepareCopyFileToPak(const FString& InMountPoint, const FPakInputPair& InFi
 		return false;
 	}
 	check(UncompressedFile.CompressedBufferSize == PaddedEncryptedFileSize);
+
 	OutNewEntry.Filename = InFile.Dest.Mid(InMountPoint.Len());
 	OutNewEntry.Info.Offset = 0; // Don't serialize offsets here.
 	OutNewEntry.Info.Size = FileSize;
@@ -1213,6 +1296,11 @@ void ProcessPakFileSpecificCommandLine(const TCHAR* CmdLine, const TArray<FStrin
 							{
 								Input.bIsDeleteRecord = true;
 							}
+							else if (Switch == TEXT("rehydrate"))
+							{
+								Input.bNeedRehydration = true;
+								CmdLineParameters.bRequiresRehydration = true;
+							}
 						}
 						else
 						{
@@ -1313,6 +1401,66 @@ void ProcessCommandLine(const TCHAR* CmdLine, const TArray<FString>& NonOptionAr
 	ProcessPakFileSpecificCommandLine(CmdLine, NonOptionArguments, Entries, CmdLineParameters);
 }
 
+bool InitializeVirtualizationSystem()
+{
+	if (UE::Virtualization::IVirtualizationSystem::IsInitialized())
+	{
+		return true;
+	}
+
+	UE_LOG(LogPakFile, Display, TEXT("Initializing the virtualization system..."));
+
+	if (FModuleManager::Get().LoadModule(TEXT("Virtualization"), ELoadModuleFlags::LogFailures) == nullptr)
+	{
+		UE_LOG(LogPakFile, Error, TEXT("Failed to load the virtualization module"));
+		return false;
+	}
+
+	IPluginManager& PluginMgr = IPluginManager::Get();
+
+	const FString PerforcePluginPath = FPaths::EnginePluginsDir() / TEXT("Developer/PerforceSourceControl/PerforceSourceControl.uplugin");
+	FText ErrorMsg;
+	if (!PluginMgr.AddToPluginsList(PerforcePluginPath, &ErrorMsg))
+	{
+		UE_LOG(LogPakFile, Error, TEXT("Failed to find 'PerforceSourceControl' plugin due to: %s"), *ErrorMsg.ToString());
+		return false;
+	}
+
+	PluginMgr.MountNewlyCreatedPlugin(TEXT("PerforceSourceControl"));
+
+	TSharedPtr<IPlugin> Plugin = PluginMgr.FindPlugin(TEXT("PerforceSourceControl"));
+	if (Plugin == nullptr || !Plugin->IsEnabled())
+	{
+		UE_LOG(LogPakFile, Error, TEXT("The 'PerforceSourceControl' plugin is disabled."));
+		return false;
+	}
+
+	FConfigFile Config;
+
+	const FString ProjectPath = FPaths::GetPath(FPaths::GetProjectFilePath());
+	const FString EngineConfigPath = FPaths::Combine(FPaths::EngineDir(), TEXT("Config/"));
+	const FString ProjectConfigPath = FPaths::Combine(ProjectPath, TEXT("Config/"));
+
+	if (!FConfigCacheIni::LoadExternalIniFile(Config, TEXT("Engine"), *EngineConfigPath, *ProjectConfigPath, true))
+	{
+		UE_LOG(LogPakFile, Error, TEXT("Failed to load config files for the project '%s"), *ProjectPath);
+		return false;
+	}
+
+		// We might need to make sure that the DDC is initialized too
+#if !USE_DDC_FOR_COMPRESSED_FILES
+	GetDerivedDataCacheRef();
+#endif //!USE_DDC_FOR_COMPRESSED_FILES
+
+	UE::Virtualization::FInitParams InitParams(FApp::GetProjectName(), Config);
+	UE::Virtualization::Initialize(InitParams, UE::Virtualization::EInitializationFlags::ForceInitialize);
+
+	UE_LOG(LogPakFile, Display, TEXT("Virtualization system initialized"));
+
+	return true;
+}
+
+
 void CollectFilesToAdd(TArray<FPakInputPair>& OutFilesToAdd, const TArray<FPakInputPair>& InEntries, const FPakOrderMap& OrderMap, const FPakCommandLineParameters& CmdLineParameters)
 {
 	TRACE_CPUPROFILER_EVENT_SCOPE(CollectFilesToAdd);
@@ -1327,6 +1475,7 @@ void CollectFilesToAdd(TArray<FPakInputPair>& OutFilesToAdd, const TArray<FPakIn
 		const FString& Source = Input.Source;
 		bool bCompression = Input.bNeedsCompression;
 		bool bEncryption = Input.bNeedEncryption;
+		bool bRehydrate = Input.bNeedRehydration;
 
 		if (Input.bIsDeleteRecord)
 		{
@@ -1387,6 +1536,8 @@ void CollectFilesToAdd(TArray<FPakInputPair>& OutFilesToAdd, const TArray<FPakIn
 				}
 				FileInput.bNeedsCompression = bCompression;
 				FileInput.bNeedEncryption = bEncryption;
+				FileInput.bNeedRehydration = bRehydrate;
+
 				if (!AddedFiles.Contains(FileInput.Source))
 				{
 					OutFilesToAdd.Add(FileInput);
@@ -1398,6 +1549,7 @@ void CollectFilesToAdd(TArray<FPakInputPair>& OutFilesToAdd, const TArray<FPakIn
 					OutFilesToAdd.Find(FileInput,FoundIndex);
 					OutFilesToAdd[FoundIndex].bNeedEncryption |= bEncryption;
 					OutFilesToAdd[FoundIndex].bNeedsCompression |= bCompression;
+					OutFilesToAdd[FoundIndex].bNeedRehydration |= bRehydrate;
 					OutFilesToAdd[FoundIndex].SuggestedOrder = FMath::Min<uint64>(OutFilesToAdd[FoundIndex].SuggestedOrder, FileInput.SuggestedOrder);
 				}
 			}
@@ -1416,6 +1568,7 @@ void CollectFilesToAdd(TArray<FPakInputPair>& OutFilesToAdd, const TArray<FPakIn
 			}
 			FileInput.bNeedEncryption = bEncryption;
 			FileInput.bNeedsCompression = bCompression;
+			FileInput.bNeedRehydration = bRehydrate;
 
 			if (AddedFiles.Contains(FileInput.Source))
 			{
@@ -1423,6 +1576,7 @@ void CollectFilesToAdd(TArray<FPakInputPair>& OutFilesToAdd, const TArray<FPakIn
 				OutFilesToAdd.Find(FileInput, FoundIndex);
 				OutFilesToAdd[FoundIndex].bNeedEncryption |= bEncryption;
 				OutFilesToAdd[FoundIndex].bNeedsCompression |= bCompression;
+				OutFilesToAdd[FoundIndex].bNeedRehydration |= bRehydrate;
 				OutFilesToAdd[FoundIndex].SuggestedOrder = FMath::Min<uint64>(OutFilesToAdd[FoundIndex].SuggestedOrder, FileInput.SuggestedOrder);
 			}
 			else
@@ -1871,7 +2025,7 @@ private:
 		bool bCopiedToPak = false;
 
 		FPakEntryPair Entry;
-		FCompressedFileBuffer CompressedFileBuffer;
+		
 		FName CompressionMethod = NAME_None;
 		int64 RealFileSize = 0;
 
@@ -1880,6 +2034,19 @@ private:
 		FGraphEventRef BeginCompressionTask;
 		FGraphEventRef EndCompressionBarrier;
 		FGraphEventRef EndCompressionTask;
+
+		FCompressedFileBuffer& AccessCompressedBuffer()
+		{
+			return CompressedFileBuffer;
+		}
+
+		const FCompressedFileBuffer& GetCompressedBuffer() const
+		{
+			return CompressedFileBuffer;
+		}
+
+	private:
+		FCompressedFileBuffer CompressedFileBuffer;
 	};
 
 	struct FOutputPakFile
@@ -1910,6 +2077,8 @@ private:
 		uint64 TotalEncryptedDataSize = 0;
 		FEventRef AllEntriesRetiredEvent;
 		int32 RetiredEntriesCount = 0;
+		uint64 RehydratedCount = 0;
+		uint64 RehydratedBytes = 0;
 	};
 
 	void BeginCompress(FOutputPakFileEntry* Entry);
@@ -2011,6 +2180,8 @@ bool FPakWriterContext::Initialize(const FPakCommandLineParameters& InCmdLinePar
 			WriterThreadFunc();
 		});
 
+
+
 	return true;
 }
 
@@ -2131,13 +2302,23 @@ void FPakWriterContext::BeginCompress(FOutputPakFileEntry* Entry)
 {
 	TRACE_CPUPROFILER_EVENT_SCOPE(BeginCompress);
 
+	if (!Entry->AccessCompressedBuffer().ReadSource(Entry->InputPair))
+	{
+		// TODO: Should we give an error?
+		Entry->bSomeCompressionSucceeded = true; // Prevent loading in EndCompress
+		Entry->EndCompressionBarrier->DispatchSubsequents();
+		return;
+	}
+
 	// don't try to compress tiny files
 	// even if they do compress, it is a bad use of decoder time
-	if (Entry->OriginalFileSize < 1024 || !Entry->InputPair.bNeedsCompression)
-	{
-		Entry->CompressedFileBuffer.ReadUncompressedFileToWorkingBuffer(Entry->InputPair);
+	if (Entry->AccessCompressedBuffer().GetFileSize() < 1024 || !Entry->InputPair.bNeedsCompression)
+	{	
+		Entry->AccessCompressedBuffer().SetSourceAsWorkingBuffer();
+
 		Entry->bSomeCompressionSucceeded = true; // Only used for logging purposes
 		Entry->EndCompressionBarrier->DispatchSubsequents();
+		// TODO: Set Entry->RealFileSize or change to accessor?
 		return;
 	}
 
@@ -2159,7 +2340,7 @@ void FPakWriterContext::BeginCompress(FOutputPakFileEntry* Entry)
 		}
 	}
 	// attempt to compress the data
-	Entry->MemoryCompressor = Entry->CompressedFileBuffer.BeginCompressFileToWorkingBuffer(Entry->InputPair, Entry->CompressionMethod, CmdLineParameters.CompressionBlockSize, Entry->EndCompressionBarrier);
+	Entry->MemoryCompressor = Entry->AccessCompressedBuffer().BeginCompressFileToWorkingBuffer(Entry->InputPair, Entry->CompressionMethod, CmdLineParameters.CompressionBlockSize, Entry->EndCompressionBarrier);
 	if (Entry->MemoryCompressor != nullptr)
 	{
 		Entry->MemoryCompressor->StartWork();
@@ -2172,25 +2353,25 @@ void FPakWriterContext::EndCompress(FOutputPakFileEntry* Entry)
 
 	const int32 MethodIndex = 0;
 	{
-		if (Entry->MemoryCompressor && Entry->CompressedFileBuffer.EndCompressFileToWorkingBuffer(Entry->InputPair, Entry->CompressionMethod, CmdLineParameters.CompressionBlockSize, *Entry->MemoryCompressor.Get()))
+		if (Entry->MemoryCompressor && Entry->AccessCompressedBuffer().EndCompressFileToWorkingBuffer(Entry->InputPair, Entry->CompressionMethod, CmdLineParameters.CompressionBlockSize, *Entry->MemoryCompressor.Get()))
 		{
 			// for modern compressors we don't want any funny heuristics turning compression on/off
 			// let the compressor decide; it has an introspective measure of whether compression is worth doing or not
-			bool bNotEnoughCompression = Entry->CompressedFileBuffer.TotalCompressedSize >= Entry->OriginalFileSize;
+			bool bNotEnoughCompression = Entry->GetCompressedBuffer().TotalCompressedSize >= Entry->OriginalFileSize;
 
 			if (Entry->CompressionMethod == NAME_Zlib)
 			{
 				// for forced-Zlib files still use the old heuristic :
 
 				// Zlib must save at least 1K regardless of percentage (for small files)
-				bNotEnoughCompression = (Entry->OriginalFileSize - Entry->CompressedFileBuffer.TotalCompressedSize) < 1024;
+				bNotEnoughCompression = (Entry->OriginalFileSize - Entry->GetCompressedBuffer().TotalCompressedSize) < 1024;
 				if (!bNotEnoughCompression)
 				{
 					// Check the compression ratio, if it's too low just store uncompressed. Also take into account read size
 					// if we still save 64KB it's probably worthwhile compressing, as that saves a file read operation in the runtime.
 					// TODO: drive this threshold from the command line
-					float PercentLess = ((float)Entry->CompressedFileBuffer.TotalCompressedSize / (Entry->OriginalFileSize / 100.f));
-					bNotEnoughCompression = (PercentLess > 90.f) && ((Entry->OriginalFileSize - Entry->CompressedFileBuffer.TotalCompressedSize) < 65536);
+					float PercentLess = ((float)Entry->GetCompressedBuffer().TotalCompressedSize / (Entry->OriginalFileSize / 100.f));
+					bNotEnoughCompression = (PercentLess > 90.f) && ((Entry->OriginalFileSize - Entry->GetCompressedBuffer().TotalCompressedSize) < 65536);
 				}
 			}
 
@@ -2201,8 +2382,8 @@ void FPakWriterContext::EndCompress(FOutputPakFileEntry* Entry)
 			}
 			else
 			{
-				Entry->Entry.Info.CompressionBlocks.AddUninitialized(Entry->CompressedFileBuffer.CompressedBlocks.Num());
-				Entry->RealFileSize = Entry->CompressedFileBuffer.TotalCompressedSize + Entry->Entry.Info.GetSerializedSize(FPakInfo::PakFile_Version_Latest);
+				Entry->Entry.Info.CompressionBlocks.AddUninitialized(Entry->GetCompressedBuffer().CompressedBlocks.Num());
+				Entry->RealFileSize = Entry->GetCompressedBuffer().TotalCompressedSize + Entry->Entry.Info.GetSerializedSize(FPakInfo::PakFile_Version_Latest);
 				Entry->Entry.Info.CompressionBlocks.Reset();
 
 				// at this point, we have successfully compressed the file, no need to continue
@@ -2214,7 +2395,7 @@ void FPakWriterContext::EndCompress(FOutputPakFileEntry* Entry)
 					// This was likely to aid with runtime error checking on mobile devices. Record how many cases of this we encountered and
 					// How much the size difference was
 					TotalFilesWithPoorForcedCompression++;
-					TotalExtraMemoryForPoorForcedCompression += FMath::Max(Entry->CompressedFileBuffer.TotalCompressedSize - Entry->OriginalFileSize, (int64)0);
+					TotalExtraMemoryForPoorForcedCompression += FMath::Max(Entry->GetCompressedBuffer().TotalCompressedSize - Entry->OriginalFileSize, (int64)0);
 				}
 			}
 		}
@@ -2222,17 +2403,22 @@ void FPakWriterContext::EndCompress(FOutputPakFileEntry* Entry)
 	Entry->MemoryCompressor.Reset();
 
 	// If no compression was able to make it small enough, or compress at all, don't compress it
-	if (!Entry->bSomeCompressionSucceeded)
+	if (Entry->bSomeCompressionSucceeded)
+	{
+		Entry->AccessCompressedBuffer().ResetSource();
+	}
+	else
 	{
 		Entry->CompressionMethod = NAME_None;
-		Entry->CompressedFileBuffer.ReadUncompressedFileToWorkingBuffer(Entry->InputPair);
+		Entry->AccessCompressedBuffer().SetSourceAsWorkingBuffer();
 	}
 }
 
 void FPakWriterContext::Retire(FOutputPakFileEntry* Entry)
 {
 	++RetiredEntriesCount;
-	Entry->CompressedFileBuffer.Empty();
+	Entry->AccessCompressedBuffer().Empty();
+
 	ScheduledFileSize -= Entry->OriginalFileSize;
 	EntryRetiredEvent->Trigger();
 	++Entry->PakFile->RetiredEntriesCount;
@@ -2322,7 +2508,7 @@ void FPakWriterContext::WriterThreadFunc()
 		}
 		
 		const FName& CompressionMethod = OutputEntry->CompressionMethod;
-		FCompressedFileBuffer& CompressedFileBuffer = OutputEntry->CompressedFileBuffer;
+		const FCompressedFileBuffer& CompressedFileBuffer = OutputEntry->GetCompressedBuffer();
 
 		if (!bDeleted)
 		{
@@ -2518,6 +2704,12 @@ void FPakWriterContext::WriterThreadFunc()
 				OutputEntry->PakFile->TotalCompressedSize += NewEntry.Info.Size;
 				OutputEntry->PakFile->TotalUncompressedSize += NewEntry.Info.UncompressedSize;
 			}
+
+			if (OutputEntry->InputPair.bNeedRehydration)
+			{
+				OutputEntry->PakFile->RehydratedCount += OutputEntry->AccessCompressedBuffer().RehydrationCount;
+				OutputEntry->PakFile->RehydratedBytes += OutputEntry->AccessCompressedBuffer().RehydrationBytes;
+			}
 		}
 		Retire(OutputEntry);
 	}
@@ -2673,6 +2865,14 @@ bool FPakWriterContext::Flush()
 			UE_LOG(LogPakFile, Display, TEXT("Compression summary: %.2f%% of original size. Compressed Size %lld bytes, Original Size %lld bytes. "), PercentLess, OutputPakFile->TotalCompressedSize, OutputPakFile->TotalUncompressedSize);
 		}
 
+		if (OutputPakFile->RehydratedCount > 0)
+		{
+			UE_LOG(LogPakFile, Display, TEXT("Asset Rehydration"));
+
+			UE_LOG(LogPakFile, Display, TEXT("  Rehydrated: %llu payloads"), OutputPakFile->RehydratedCount);
+			UE_LOG(LogPakFile, Display, TEXT("  Rehydrated: %llu bytes (%.2fMB)"), OutputPakFile->RehydratedBytes, (float)OutputPakFile->RehydratedBytes / 1024.0f / 1024.0f);
+		}
+
 		if (OutputPakFile->TotalEncryptedDataSize)
 		{
 			UE_LOG(LogPakFile, Display, TEXT("Encryption - ENABLED"));
@@ -5670,6 +5870,14 @@ bool ExecuteUnrealPak(const TCHAR* CmdLine)
 			FPakCommandLineParameters CmdLineParameters = CommonCmdLineParameters;
 			ProcessPakFileSpecificCommandLine(*CreatePakCommand, NonOptionArgumentsForPakFile, Entries, CmdLineParameters);
 			
+			if (CmdLineParameters.bRequiresRehydration)
+			{
+				if (!InitializeVirtualizationSystem())
+				{
+					return false;
+				}
+			}
+
 			FKeyChain KeyChainForPakFile = KeyChain;
 			FString EncryptionKeyOverrideGuidString;
 			FGuid EncryptionKeyOverrideGuid;
diff --git a/Engine/Source/Programs/UnrealPak/UnrealPak.Build.cs b/Engine/Source/Programs/UnrealPak/UnrealPak.Build.cs
--- a/Engine/Source/Programs/UnrealPak/UnrealPak.Build.cs
+++ b/Engine/Source/Programs/UnrealPak/UnrealPak.Build.cs
@@ -8,13 +8,29 @@ public class UnrealPak : ModuleRules
 	{
 		PublicIncludePaths.Add("Runtime/Launch/Public");
 
-		PrivateDependencyModuleNames.AddRange(new string[] { "Core", "CoreUObject", "AssetRegistry", "PakFile", "Json", "Projects", "PakFileUtilities", "RSA", "ApplicationCore" });
+		PrivateDependencyModuleNames.AddRange(new string[] { 
+			"Core", 
+			"CoreUObject", 
+			"AssetRegistry", 
+			"PakFile", 
+			"Json", 
+			"Projects", 
+			"PakFileUtilities", 
+			"RSA", 
+			"ApplicationCore" 
+		});
 
-		PrivateIncludePaths.Add("Runtime/Launch/Private");      // For LaunchEngineLoop.cpp include
+		PrivateIncludePaths.Add("Runtime/Launch/Private"); // For LaunchEngineLoop.cpp include
 
-        PrivateIncludePathModuleNames.AddRange(
-            new string[] {
-                "Json"
-        });
-    }
+		PrivateIncludePathModuleNames.AddRange(new string[] {
+			"Json"
+		});
+
+		if (Target.bBuildWithEditorOnlyData)
+		{
+			DynamicallyLoadedModuleNames.AddRange(new string[] {
+				"PerforceSourceControl"
+			});
+		}
+	}
 }
