PakFile: Add functionality to sign a pak as a separate step

diff --git a/Engine/Source/Developer/PakFileUtilities/Private/PakFileUtilities.cpp b/Engine/Source/Developer/PakFileUtilities/Private/PakFileUtilities.cpp
--- a/Engine/Source/Developer/PakFileUtilities/Private/PakFileUtilities.cpp
+++ b/Engine/Source/Developer/PakFileUtilities/Private/PakFileUtilities.cpp
@@ -72,6 +72,68 @@ int64 GDDCHits = 0;
 int64 GDDCMisses = 0;
 #endif
 
+
+bool SignPakFile(const FString& InPakFilename, const FRSAKeyHandle InSigningKey)
+{
+	FString SignatureFilename(FPaths::ChangeExtension(InPakFilename, TEXT(".sig")));
+	TUniquePtr<FArchive> PakFile(IFileManager::Get().CreateFileReader(*InPakFilename));
+
+	const int64 TotalSize = PakFile->TotalSize();
+	bool bFoundMagic = false;
+	FPakInfo PakInfo;
+	const int64 FileInfoSize = PakInfo.GetSerializedSize(FPakInfo::PakFile_Version_Latest);
+	if (TotalSize >= FileInfoSize)
+	{
+		const int64 FileInfoPos = TotalSize - FileInfoSize;
+		PakFile->Seek(FileInfoPos);
+		PakInfo.Serialize(*PakFile.Get(), FPakInfo::PakFile_Version_Latest);
+		bFoundMagic = (PakInfo.Magic == FPakInfo::PakFile_Magic);
+	}
+
+	// TODO: This should probably mimic the logic of FPakFile::Initialize which iterates back through pak versions until it find a compatible one
+	if (!bFoundMagic)
+	{
+		return false;
+	}
+
+	TArray<uint8> ComputedSignatureData;
+	ComputedSignatureData.Append(PakInfo.IndexHash.Hash, UE_ARRAY_COUNT(FSHAHash::Hash));
+
+	const uint64 BlockSize = FPakInfo::MaxChunkDataSize;
+	uint64 Remaining = PakFile->TotalSize();
+
+	PakFile->Seek(0);
+
+	TArray<uint8> Buffer;
+	Buffer.SetNum(BlockSize);
+
+	TArray<TPakChunkHash> Hashes;
+	Hashes.Empty((uint64)(PakFile->TotalSize() / BlockSize));
+
+	while (Remaining > 0)
+	{
+		const uint64 CurrentBlockSize = FMath::Min(BlockSize, Remaining);
+		Remaining -= CurrentBlockSize;
+
+		PakFile->Serialize(Buffer.GetData(), CurrentBlockSize);
+		Hashes.Add(ComputePakChunkHash(Buffer.GetData(), CurrentBlockSize));
+	}
+
+	FPakSignatureFile Signatures;
+	Signatures.SetChunkHashesAndSign(Hashes, ComputedSignatureData, InSigningKey);
+
+	TUniquePtr<FArchive> SignatureFile(IFileManager::Get().CreateFileWriter(*SignatureFilename));
+	Signatures.Serialize(*SignatureFile.Get());
+
+	return true;
+}
+
+bool SignIOStoreContainer(FArchive& ContainerAr)
+{
+	return true;
+}
+
+
 class FMemoryCompressor;
 
 /**
diff --git a/Engine/Source/Developer/PakFileUtilities/Public/PakFileUtilities.h b/Engine/Source/Developer/PakFileUtilities/Public/PakFileUtilities.h
--- a/Engine/Source/Developer/PakFileUtilities/Public/PakFileUtilities.h
+++ b/Engine/Source/Developer/PakFileUtilities/Public/PakFileUtilities.h
@@ -10,6 +10,7 @@
 #include "Math/NumericLimits.h"
 #include "Math/UnrealMathSSE.h"
 #include "Misc/Optional.h"
+#include "Misc/IEngineCrypto.h"
 #include "Templates/Tuple.h"
 
 class FArchive;
@@ -142,3 +143,6 @@ struct FPakFooterInfo
 
 /** Write the index and other data at the end of a pak file after all the entries */
 PAKFILEUTILITIES_API void WritePakFooter(FArchive& PakHandle, FPakFooterInfo& FooterInfo);
+
+/** Take an existing pak file and regenerate the signature file */
+PAKFILEUTILITIES_API bool SignPakFile(const FString& InPakFilename, const FRSAKeyHandle InSigningKey);
\ No newline at end of file
