UnrealPak: Fix unavoidable crash when using -repack (fallout from cl 18836432).
Let CollectFilesToAdd fallback to the old logic of using the source file name when the FPakInputPair.Dest specifies a directory, i.e. is ending with a "/".
Let ExtractFilesFromPak start using destination file names just as ProcessCommandLine reading from response files is doing so that -Extract with -responsefile will generate the same format as UAT with file names.
Remove dead code in ProcessPakFileSpecificCommandLine.

diff --git a/Engine/Source/Developer/PakFileUtilities/Private/PakFileUtilities.cpp b/Engine/Source/Developer/PakFileUtilities/Private/PakFileUtilities.cpp
--- a/Engine/Source/Developer/PakFileUtilities/Private/PakFileUtilities.cpp
+++ b/Engine/Source/Developer/PakFileUtilities/Private/PakFileUtilities.cpp
@@ -1395,16 +1395,10 @@ void ProcessPakFileSpecificCommandLine(const TCHAR* CmdLine, const TArray<FStrin
 					Input.Dest = Input.Source;
 				}
 
-				FString DestPath;
 				{
 					TRACE_CPUPROFILER_EVENT_SCOPE(OtherPathStuff);
 					FPaths::NormalizeFilename(Input.Source);
 					FPaths::NormalizeFilename(Input.Dest);
-					DestPath = FPaths::GetPath(Input.Dest);
-				}
-				{
-					TRACE_CPUPROFILER_EVENT_SCOPE(MakeDirectoryFromPath);
-					FPakFile::MakeDirectoryFromPath(DestPath);
 				}
 
 				Input.bNeedsCompression |= bCompress;
@@ -1622,7 +1616,9 @@ void CollectFilesToAdd(TArray<FPakInputPair>& OutFilesToAdd, const TArray<FPakIn
 			FPakInputPair FileInput;
 			FileInput.Source = Input.Source;
 			FPaths::MakeStandardFilename(FileInput.Source);
-			FileInput.Dest = Input.Dest;
+			FileInput.Dest = Input.Dest.EndsWith(TEXT("/")) ?
+				FileInput.Source.Replace(*Directory, *Input.Dest, ESearchCase::IgnoreCase) :
+				Input.Dest;
 			uint64 FileOrder = OrderMap.GetFileOrder(FileInput.Dest, CmdLineParameters.bFallbackOrderForNonUassetFiles, &FileInput.bIsInPrimaryOrder);
 			if (FileOrder != MAX_uint64)
 			{
@@ -4052,9 +4048,7 @@ bool ExtractFilesFromPak(const TCHAR* InPakFilename, TMap<FString, FFileInfo>& I
 								Input.Source = DestFilename;
 								FPaths::NormalizeFilename(Input.Source);
 
-								Input.Dest = PakFile.GetMountPoint() + FPaths::GetPath(EntryFileName);
-								FPaths::NormalizeFilename(Input.Dest);
-								FPakFile::MakeDirectoryFromPath(Input.Dest);
+								Input.Dest = PakFile.GetMountPoint() / EntryFileName;
 
 								Input.bNeedsCompression = Entry.CompressionMethodIndex != 0;
 								Input.bNeedEncryption = Entry.IsEncrypted();
