UnrealPak: Fix check triggering for modifying compression tasks array while iterating over it

diff --git a/Engine/Source/Developer/PakFileUtilities/Private/PakFileUtilities.cpp b/Engine/Source/Developer/PakFileUtilities/Private/PakFileUtilities.cpp
--- a/Engine/Source/Developer/PakFileUtilities/Private/PakFileUtilities.cpp
+++ b/Engine/Source/Developer/PakFileUtilities/Private/PakFileUtilities.cpp
@@ -242,11 +242,12 @@ public:
 		}
 		else
 		{
-			RemainingTasksCounter = BlockCompressAsyncTasks.Num();
-
-			for (FAsyncTask<FBlockCompressTask>* AsyncTask : BlockCompressAsyncTasks)
+			int32 AsyncTasksCount = BlockCompressAsyncTasks.Num();
+			RemainingTasksCounter = AsyncTasksCount;
+			
+			for (int32 TaskIndex = 0; TaskIndex < AsyncTasksCount; ++TaskIndex)
 			{
-				AsyncTask->StartBackgroundTask();
+				BlockCompressAsyncTasks[TaskIndex]->StartBackgroundTask();
 			}
 		}
 	}
