UnrealPak: Improve diffutils
Add support for diffing pak directories for pak diff.
Add support for different cryptokeys for pak diff and legacy iostore diff.

#rnx

diff --git a/Engine/Source/Developer/IoStoreUtilities/Private/IoStoreUtilities.cpp b/Engine/Source/Developer/IoStoreUtilities/Private/IoStoreUtilities.cpp
--- a/Engine/Source/Developer/IoStoreUtilities/Private/IoStoreUtilities.cpp
+++ b/Engine/Source/Developer/IoStoreUtilities/Private/IoStoreUtilities.cpp
@@ -5497,12 +5497,12 @@ static int32 Diff(
 	return 0;
 }
 
-bool LegacyDiffIoStoreContainers(const TCHAR* InContainerFilename1, const TCHAR* InContainerFilename2, bool bInLogUniques1, bool bInLogUniques2, const FKeyChain& InKeyChain)
+bool LegacyDiffIoStoreContainers(const TCHAR* InContainerFilename1, const TCHAR* InContainerFilename2, bool bInLogUniques1, bool bInLogUniques2, const FKeyChain& InKeyChain1, const FKeyChain* InKeyChain2)
 {
 	TGuardValue<ELogTimes::Type> DisableLogTimes(GPrintLogTimes, ELogTimes::None);
 	UE_LOG(LogIoStore, Log, TEXT("FileEventType, FileName, Size1, Size2"));
 
-	TUniquePtr<FIoStoreReader> Reader1 = CreateIoStoreReader(InContainerFilename1, InKeyChain);
+	TUniquePtr<FIoStoreReader> Reader1 = CreateIoStoreReader(InContainerFilename1, InKeyChain1);
 	if (!Reader1.IsValid())
 	{
 		return false;
@@ -5513,7 +5513,7 @@ bool LegacyDiffIoStoreContainers(const TCHAR* InContainerFilename1, const TCHAR*
 		UE_LOG(LogIoStore, Warning, TEXT("Missing directory index for container '%s'"), InContainerFilename1);
 	}
 
-	TUniquePtr<FIoStoreReader> Reader2 = CreateIoStoreReader(InContainerFilename2, InKeyChain);
+	TUniquePtr<FIoStoreReader> Reader2 = CreateIoStoreReader(InContainerFilename2, InKeyChain2 ? *InKeyChain2 : InKeyChain1);
 	if (!Reader2.IsValid())
 	{
 		return false;
@@ -6671,7 +6671,8 @@ int32 CreateIoStoreContainerFiles(const TCHAR* CmdLine)
 		FParse::Value(FCommandLine::Get(), TEXT("DumpToFile="), OutPath);
 
 		FString CryptoKeysCacheFilename;
-		if (FParse::Value(CmdLine, TEXT("SourceCryptoKeys="), CryptoKeysCacheFilename))
+		if (FParse::Value(CmdLine, TEXT("CryptoKeys="), CryptoKeysCacheFilename) ||
+			FParse::Value(CmdLine, TEXT("SourceCryptoKeys="), CryptoKeysCacheFilename))
 		{
 			UE_LOG(LogIoStore, Display, TEXT("Parsing source crypto keys from '%s'"), *CryptoKeysCacheFilename);
 			KeyChainUtilities::LoadKeyChainFromFile(CryptoKeysCacheFilename, SourceKeyChain);
@@ -6682,6 +6683,10 @@ int32 CreateIoStoreContainerFiles(const TCHAR* CmdLine)
 			UE_LOG(LogIoStore, Display, TEXT("Parsing target crypto keys from '%s'"), *CryptoKeysCacheFilename);
 			KeyChainUtilities::LoadKeyChainFromFile(CryptoKeysCacheFilename, TargetKeyChain);
 		}
+		else
+		{
+			TargetKeyChain = SourceKeyChain;
+		}
 
 		EChunkTypeFilter ChunkTypeFilter = EChunkTypeFilter::None;
 		if (FParse::Param(FCommandLine::Get(), TEXT("FilterBulkData")))
diff --git a/Engine/Source/Developer/IoStoreUtilities/Public/IoStoreUtilities.h b/Engine/Source/Developer/IoStoreUtilities/Public/IoStoreUtilities.h
--- a/Engine/Source/Developer/IoStoreUtilities/Public/IoStoreUtilities.h
+++ b/Engine/Source/Developer/IoStoreUtilities/Public/IoStoreUtilities.h
@@ -27,7 +27,8 @@ IOSTOREUTILITIES_API bool LegacyDiffIoStoreContainers(
 	const TCHAR* InContainerFilename2,
 	bool bInLogUniques1,
 	bool bInLogUniques2,
-	const FKeyChain& InKeyChain);
+	const FKeyChain& InKeyChain1,
+	const FKeyChain* InKeyChain2 = nullptr);
 
 IOSTOREUTILITIES_API bool ExtractFilesFromIoStoreContainer(
 	const TCHAR* InContainerFilename,
diff --git a/Engine/Source/Developer/PakFileUtilities/Private/PakFileUtilities.cpp b/Engine/Source/Developer/PakFileUtilities/Private/PakFileUtilities.cpp
--- a/Engine/Source/Developer/PakFileUtilities/Private/PakFileUtilities.cpp
+++ b/Engine/Source/Developer/PakFileUtilities/Private/PakFileUtilities.cpp
@@ -4166,7 +4166,7 @@ bool DumpPakInfo(const FString& InPakFilename, const FKeyChain& InKeyChain)
 	return true;
 }
 
-bool DiffFilesInPaks(const FString& InPakFilename1, const FString& InPakFilename2, const bool bLogUniques1, const bool bLogUniques2, const FKeyChain& InKeyChain)
+bool DiffFilesInPaks(const FString& InPakFilename1, const FString& InPakFilename2, const bool bLogUniques1, const bool bLogUniques2, const FKeyChain& InKeyChain1, const FKeyChain& InKeyChain2)
 {
 	int32 NumUniquePAK1 = 0;
 	int32 NumUniquePAK2 = 0;
@@ -4253,20 +4253,20 @@ bool DiffFilesInPaks(const FString& InPakFilename1, const FString& InPakFilename
 
 				if (EntryInfo1.CompressionMethodIndex == 0)
 				{
-					BufferedCopyFile(PAKWriter1, PakReader1.GetArchive(), PakFile1, Entry1, Buffer, BufferSize, InKeyChain);
+					BufferedCopyFile(PAKWriter1, PakReader1.GetArchive(), PakFile1, Entry1, Buffer, BufferSize, InKeyChain1);
 				}
 				else
 				{
-					UncompressCopyFile(PAKWriter1, PakReader1.GetArchive(), Entry1, PersistantCompressionBuffer, CompressionBufferSize, InKeyChain, PakFile1);
+					UncompressCopyFile(PAKWriter1, PakReader1.GetArchive(), Entry1, PersistantCompressionBuffer, CompressionBufferSize, InKeyChain1, PakFile1);
 				}
 
 				if (EntryInfo2.CompressionMethodIndex == 0)
 				{
-					BufferedCopyFile(PAKWriter2, PakReader2.GetArchive(), PakFile2, Entry2, Buffer, BufferSize, InKeyChain);
+					BufferedCopyFile(PAKWriter2, PakReader2.GetArchive(), PakFile2, Entry2, Buffer, BufferSize, InKeyChain2);
 				}
 				else
 				{
-					UncompressCopyFile(PAKWriter2, PakReader2.GetArchive(), Entry2, PersistantCompressionBuffer, CompressionBufferSize, InKeyChain, PakFile2);
+					UncompressCopyFile(PAKWriter2, PakReader2.GetArchive(), Entry2, PersistantCompressionBuffer, CompressionBufferSize, InKeyChain2, PakFile2);
 				}
 
 				if (FMemory::Memcmp(PAKWriter1.GetData(), PAKWriter2.GetData(), EntryInfo1.UncompressedSize) != 0)
@@ -4310,11 +4310,10 @@ bool DiffFilesInPaks(const FString& InPakFilename1, const FString& InPakFilename
 		FMemory::Free(Buffer);
 		Buffer = nullptr;
 
-		UE_LOG(LogPakFile, Log, TEXT("Comparison complete"));
-		UE_LOG(LogPakFile, Log, TEXT("Unique to first pak: %i, Unique to second pak: %i, Num Different: %i, NumEqual: %i"), NumUniquePAK1, NumUniquePAK2, NumDifferentContents, NumEqualContents);
+		UE_LOG(LogPakFile, Display, TEXT("Unique to first pak: %i, Unique to second pak: %i, Num Different: %i, NumEqual: %i"), NumUniquePAK1, NumUniquePAK2, NumDifferentContents, NumEqualContents);
 		return true;
 	}
-	else if (LegacyDiffIoStoreContainers(*InPakFilename1, *InPakFilename2, bLogUniques1, bLogUniques2, InKeyChain))
+	else if (LegacyDiffIoStoreContainers(*InPakFilename1, *InPakFilename2, bLogUniques1, bLogUniques2, InKeyChain1, &InKeyChain2))
 	{
 		return true;
 	}
@@ -5526,19 +5525,96 @@ bool ExecuteUnrealPak(const TCHAR* CmdLine)
 	{
 		if(NonOptionArguments.Num() != 2)
 		{
-			UE_LOG(LogPakFile,Error,TEXT("Incorrect arguments. Expected: -Diff <PakFile1> <PakFile2> [-NoUniques] [-NoUniquesFile1] [-NoUniquesFile2]"));
+			UE_LOG(LogPakFile,Error,TEXT("Incorrect arguments. Expected: -Diff <PakFile1|PakPath1> <PakFile2|PakPath2> [-NoUniques] [-NoUniquesFile1] [-NoUniquesFile2] -CryptoKeys2=<cryptofile>"));
 			return false;
 		}
 
-		FString PakFilename1 = GetPakPath(*NonOptionArguments[0], false);
-		FString PakFilename2 = GetPakPath(*NonOptionArguments[1], false);
+		auto GetPakFiles = [](const FString& PakFileOrFolder, FString& OutFolder, TArray<FString>& OutFiles)
+		{
+			if (FPaths::DirectoryExists(PakFileOrFolder))
+			{
+				OutFolder = PakFileOrFolder;
+				IFileManager::Get().FindFiles(OutFiles, *PakFileOrFolder, TEXT(".pak"));
+			}
+			else
+			{
+				FString PakFile = GetPakPath(*PakFileOrFolder, false);
+				OutFolder = FPaths::GetPath(PakFile);
+				OutFiles.Emplace(FPaths::GetCleanFilename(PakFile));
+			}
+		};
+		FString SourceFolderName;
+		FString TargetFolderName;
+		TArray<FString> SourcePakFiles;
+		TArray<FString> TargetPakFiles;
+		GetPakFiles(NonOptionArguments[0], SourceFolderName, SourcePakFiles);
+		GetPakFiles(NonOptionArguments[1], TargetFolderName, TargetPakFiles);
 
 		// Allow the suppression of unique file logging for one or both files
 		const bool bLogUniques = !FParse::Param(CmdLine, TEXT("nouniques"));
 		const bool bLogUniques1 = bLogUniques && !FParse::Param(CmdLine, TEXT("nouniquesfile1"));
 		const bool bLogUniques2 = bLogUniques && !FParse::Param(CmdLine, TEXT("nouniquesfile2"));
 
-		return DiffFilesInPaks(PakFilename1, PakFilename2, bLogUniques1, bLogUniques2, KeyChain);
+		FString CryptoKeysFilename2;
+		FKeyChain KeyChain2;
+		if (FParse::Value(FCommandLine::Get(), TEXT("CryptoKeys2="), CryptoKeysFilename2) ||
+			FParse::Value(FCommandLine::Get(), TEXT("TargetCryptoKeys="), CryptoKeysFilename2))
+		{
+			KeyChainUtilities::LoadKeyChainFromFile(CryptoKeysFilename2, KeyChain2);
+			KeyChainUtilities::ApplyEncryptionKeys(KeyChain2);
+		}
+		else
+		{
+			KeyChain2 = KeyChain;
+		}
+
+		TSet<FString> SourcePakFileSet;
+		TSet<FString> TargetPakFileSet;
+		SourcePakFileSet.Append(SourcePakFiles);
+		TargetPakFileSet.Append(TargetPakFiles);
+
+		UE_LOG(LogPakFile, Display, TEXT("Diff Source Pak Folder: %s - %d pak file(s)"),
+			*SourceFolderName, SourcePakFiles.Num());
+		UE_LOG(LogPakFile, Display, TEXT("Diff Target Pak Folder: %s - %d pak file(s)"),
+			*TargetFolderName, TargetPakFiles.Num());
+
+		for (int32 I = SourcePakFiles.Num() - 1; I >= 0; --I)
+		{
+			if (!TargetPakFileSet.Contains(SourcePakFiles[I]))
+			{
+				if (bLogUniques1)
+				{
+					UE_LOG(LogPakFile, Display, TEXT("Source PakFile '%s' is missing in target folder"),
+						*SourcePakFiles[I]);
+				}
+				SourcePakFiles.RemoveAtSwap(I, 1, false);
+			}
+		}
+		if (bLogUniques2)
+		{
+			for (int32 I = TargetPakFiles.Num() - 1; I >= 0; --I)
+			{
+				if (!SourcePakFileSet.Contains(TargetPakFiles[I]))
+				{
+					UE_LOG(LogPakFile, Display, TEXT("Target PakFile '%s' is missing in source folder"),
+						*TargetPakFiles[I]);
+				}
+			}
+		}
+
+		bool Success = true;
+		FString SourceFile;
+		FString TargetFile;
+		for (const FString& PakFile : SourcePakFiles)
+		{
+			SourceFile = SourceFolderName / PakFile;
+			SourceFile.ReplaceInline(TEXT("/"), TEXT("\\"));
+			TargetFile = TargetFolderName / PakFile;
+			TargetFile.ReplaceInline(TEXT("/"), TEXT("\\"));
+			UE_LOG(LogPakFile, Display, TEXT("Diffing PakFile '%s'"), *PakFile);
+			Success &= DiffFilesInPaks(SourceFile, TargetFile, bLogUniques1, bLogUniques2, KeyChain, KeyChain2);
+		}
+		return Success;
 	}
 
 	if (FParse::Param(CmdLine, TEXT("Extract")))
