Tidy up pak log spam a bit when making 1000s of pak files.

diff --git a/Engine/Source/Developer/PakFileUtilities/Private/PakFileUtilities.cpp b/Engine/Source/Developer/PakFileUtilities/Private/PakFileUtilities.cpp
--- a/Engine/Source/Developer/PakFileUtilities/Private/PakFileUtilities.cpp
+++ b/Engine/Source/Developer/PakFileUtilities/Private/PakFileUtilities.cpp
@@ -2801,7 +2801,17 @@ bool FPakWriterContext::Flush()
 		}
 
 		TRACE_CPUPROFILER_EVENT_SCOPE(FinalizePakFile);
-		UE_LOG(LogPakFile, Display, TEXT("Created pak file: %s"), *OutputPakFile->Filename);
+
+		bool bPakFileIsEmpty = OutputPakFile->Index.Num() == 0;
+
+		if (bPakFileIsEmpty)
+		{
+			UE_LOG(LogPakFile, Display, TEXT("Created empty pak file: %s"), *OutputPakFile->Filename);
+		}
+		else
+		{
+			UE_LOG(LogPakFile, Display, TEXT("Created pak file: %s"), *OutputPakFile->Filename);
+		}
 
 		for (FOutputPakFileEntry& OutputEntry : OutputPakFile->Entries)
 		{
@@ -2903,62 +2913,68 @@ bool FPakWriterContext::Flush()
 			OutputPakFile->PakFileRegionsHandle.Reset();
 		}
 
-		// log per-compressor stats :
-		for (int32 MethodIndex = 0; MethodIndex < CompressionFormatsAndNone.Num(); MethodIndex++)
-		{
-			FName CompressionMethod = CompressionFormatsAndNone[MethodIndex];
-			UE_LOG(LogPakFile, Display, TEXT("CompressionFormat %d [%s] : %d files, %lld -> %lld bytes"), MethodIndex, *(CompressionMethod.ToString()),
-				OutputPakFile->Compressor_Stat_Count[MethodIndex],
-				OutputPakFile->Compressor_Stat_RawBytes[MethodIndex],
-				OutputPakFile->Compressor_Stat_CompBytes[MethodIndex]
-			);
-		}
-
-		UE_LOG(LogPakFile, Display, TEXT("Added %d files, %lld bytes total"), OutputPakFile->Index.Num(), TotalPakFileSize);
-		UE_LOG(LogPakFile, Display, TEXT("PrimaryIndex size: %d bytes"), Footer.PrimaryIndexSize);
-		UE_LOG(LogPakFile, Display, TEXT("PathHashIndex size: %d bytes"), Footer.PathHashIndexSize);
-		UE_LOG(LogPakFile, Display, TEXT("FullDirectoryIndex size: %d bytes"), Footer.FullDirectoryIndexSize);
-		if (OutputPakFile->TotalUncompressedSize)
+		if (bPakFileIsEmpty == false)
 		{
-			float PercentLess = ((float)OutputPakFile->TotalCompressedSize / ((float)OutputPakFile->TotalUncompressedSize / 100.f));
-			UE_LOG(LogPakFile, Display, TEXT("Compression summary: %.2f%% of original size. Compressed Size %lld bytes, Original Size %lld bytes. "), PercentLess, OutputPakFile->TotalCompressedSize, OutputPakFile->TotalUncompressedSize);
-		}
+			// log per-compressor stats :
+			for (int32 MethodIndex = 0; MethodIndex < CompressionFormatsAndNone.Num(); MethodIndex++)
+			{
+				if (OutputPakFile->Compressor_Stat_Count[MethodIndex])
+				{
+					FName CompressionMethod = CompressionFormatsAndNone[MethodIndex];
+					UE_LOG(LogPakFile, Display, TEXT("CompressionFormat %d [%s] : %d files, %lld -> %lld bytes"), MethodIndex, *(CompressionMethod.ToString()),
+						OutputPakFile->Compressor_Stat_Count[MethodIndex],
+						OutputPakFile->Compressor_Stat_RawBytes[MethodIndex],
+						OutputPakFile->Compressor_Stat_CompBytes[MethodIndex]
+					);
+				}
+			}
 
-		if (OutputPakFile->RehydratedCount > 0)
-		{
-			UE_LOG(LogPakFile, Display, TEXT("Asset Rehydration"));
+			UE_LOG(LogPakFile, Display, TEXT("Added %d files, %lld bytes total"), OutputPakFile->Index.Num(), TotalPakFileSize);
+			UE_LOG(LogPakFile, Display, TEXT("PrimaryIndex size: %d bytes"), Footer.PrimaryIndexSize);
+			UE_LOG(LogPakFile, Display, TEXT("PathHashIndex size: %d bytes"), Footer.PathHashIndexSize);
+			UE_LOG(LogPakFile, Display, TEXT("FullDirectoryIndex size: %d bytes"), Footer.FullDirectoryIndexSize);
+			if (OutputPakFile->TotalUncompressedSize)
+			{
+				float PercentLess = ((float)OutputPakFile->TotalCompressedSize / ((float)OutputPakFile->TotalUncompressedSize / 100.f));
+				UE_LOG(LogPakFile, Display, TEXT("Compression summary: %.2f%% of original size. Compressed Size %lld bytes, Original Size %lld bytes. "), PercentLess, OutputPakFile->TotalCompressedSize, OutputPakFile->TotalUncompressedSize);
+			}
 
-			UE_LOG(LogPakFile, Display, TEXT("  Rehydrated: %llu payloads"), OutputPakFile->RehydratedCount);
-			UE_LOG(LogPakFile, Display, TEXT("  Rehydrated: %llu bytes (%.2fMB)"), OutputPakFile->RehydratedBytes, (float)OutputPakFile->RehydratedBytes / 1024.0f / 1024.0f);
-		}
+			if (OutputPakFile->RehydratedCount > 0)
+			{
+				UE_LOG(LogPakFile, Display, TEXT("Asset Rehydration"));
 
-		if (OutputPakFile->TotalEncryptedDataSize)
-		{
-			UE_LOG(LogPakFile, Display, TEXT("Encryption - ENABLED"));
-			UE_LOG(LogPakFile, Display, TEXT("  Files: %d"), OutputPakFile->TotalEncryptedFiles);
+				UE_LOG(LogPakFile, Display, TEXT("  Rehydrated: %llu payloads"), OutputPakFile->RehydratedCount);
+				UE_LOG(LogPakFile, Display, TEXT("  Rehydrated: %llu bytes (%.2fMB)"), OutputPakFile->RehydratedBytes, (float)OutputPakFile->RehydratedBytes / 1024.0f / 1024.0f);
+			}
 
-			if (OutputPakFile->Info.bEncryptedIndex)
+			if (OutputPakFile->TotalEncryptedDataSize)
 			{
-				UE_LOG(LogPakFile, Display, TEXT("  Index: Encrypted (%d bytes, %.2fMB)"), OutputPakFile->Info.IndexSize, (float)OutputPakFile->Info.IndexSize / 1024.0f / 1024.0f);
+				UE_LOG(LogPakFile, Display, TEXT("Encryption - ENABLED"));
+				UE_LOG(LogPakFile, Display, TEXT("  Files: %d"), OutputPakFile->TotalEncryptedFiles);
+
+				if (OutputPakFile->Info.bEncryptedIndex)
+				{
+					UE_LOG(LogPakFile, Display, TEXT("  Index: Encrypted (%d bytes, %.2fMB)"), OutputPakFile->Info.IndexSize, (float)OutputPakFile->Info.IndexSize / 1024.0f / 1024.0f);
+				}
+				else
+				{
+					UE_LOG(LogPakFile, Display, TEXT("  Index: Unencrypted"));
+				}
+
+
+				UE_LOG(LogPakFile, Display, TEXT("  Total: %d bytes (%.2fMB)"), OutputPakFile->TotalEncryptedDataSize, (float)OutputPakFile->TotalEncryptedDataSize / 1024.0f / 1024.0f);
 			}
 			else
 			{
-				UE_LOG(LogPakFile, Display, TEXT("  Index: Unencrypted"));
+				UE_LOG(LogPakFile, Display, TEXT("Encryption - DISABLED"));
 			}
 
-
-			UE_LOG(LogPakFile, Display, TEXT("  Total: %d bytes (%.2fMB)"), OutputPakFile->TotalEncryptedDataSize, (float)OutputPakFile->TotalEncryptedDataSize / 1024.0f / 1024.0f);
-		}
-		else
-		{
-			UE_LOG(LogPakFile, Display, TEXT("Encryption - DISABLED"));
-		}
-
-		if (OutputPakFile->TotalEncryptedFiles < OutputPakFile->TotalRequestedEncryptedFiles)
-		{
-			UE_LOG(LogPakFile, Display, TEXT("%d files requested encryption, but no AES key was supplied! Encryption was skipped for these files"), OutputPakFile->TotalRequestedEncryptedFiles);
+			if (OutputPakFile->TotalEncryptedFiles < OutputPakFile->TotalRequestedEncryptedFiles)
+			{
+				UE_LOG(LogPakFile, Display, TEXT("%d files requested encryption, but no AES key was supplied! Encryption was skipped for these files"), OutputPakFile->TotalRequestedEncryptedFiles);
+			}
+			UE_LOG(LogPakFile, Display, TEXT(""));
 		}
-		UE_LOG(LogPakFile, Display, TEXT(""));
 	}
 
 	WriterThread.Wait();
