Check for IO errors when signing paks.
Fix pak files not being signed.

diff --git a/Engine/Source/Developer/PakFileUtilities/Private/PakFileUtilities.cpp b/Engine/Source/Developer/PakFileUtilities/Private/PakFileUtilities.cpp
--- a/Engine/Source/Developer/PakFileUtilities/Private/PakFileUtilities.cpp
+++ b/Engine/Source/Developer/PakFileUtilities/Private/PakFileUtilities.cpp
@@ -126,7 +126,8 @@ bool SignPakFile(const FString& InPakFilename, const FRSAKeyHandle InSigningKey)
 	TUniquePtr<FArchive> SignatureFile(IFileManager::Get().CreateFileWriter(*SignatureFilename));
 	Signatures.Serialize(*SignatureFile.Get());
 
-	return true;
+	const bool bSuccess = !PakFile->IsError() && !SignatureFile->IsError();
+	return bSuccess;
 }
 
 bool SignIOStoreContainer(FArchive& ContainerAr)
