PAK - flush pending file writers before exiting UnrealPak

diff --git a/Engine/Source/Developer/PakFileUtilities/Private/PakFileUtilities.cpp b/Engine/Source/Developer/PakFileUtilities/Private/PakFileUtilities.cpp
--- a/Engine/Source/Developer/PakFileUtilities/Private/PakFileUtilities.cpp
+++ b/Engine/Source/Developer/PakFileUtilities/Private/PakFileUtilities.cpp
@@ -3052,6 +3052,7 @@ bool CreatePakFile(const TCHAR* Filename, const TArray<FPakInputPair>& FilesToAd
 	
 	if (!PakWriterContext.AddPakFile(Filename, FilesToAdd, InKeyChain))
 	{
+		PakWriterContext.Flush();
 		return false;
 	}
 
@@ -6043,6 +6044,7 @@ bool ExecuteUnrealPak(const TCHAR* CmdLine)
 		KeyChainUtilities::ApplyEncryptionKeys(KeyChain);
 
 		TArray<FString> TempOutputDirectoriesToDelete;
+		bool bResult = true;
 		for (const FString& CreatePakCommand : CreatePakCommandsList)
 		{
 			TArray<FString> NonOptionArgumentsForPakFile;
@@ -6067,7 +6069,8 @@ bool ExecuteUnrealPak(const TCHAR* CmdLine)
 			{
 				if (!InitializeVirtualizationSystem())
 				{
-					return false;
+					bResult = false;
+					break;
 				}
 			}
 
@@ -6164,11 +6167,12 @@ bool ExecuteUnrealPak(const TCHAR* CmdLine)
 
 			if (!PakWriterContext.AddPakFile(*PakFilename, FilesToAdd, KeyChainForPakFile))
 			{
-				return false;
+				bResult = false;
+				break;
 			}
 		}
 
-		bool bResult = PakWriterContext.Flush();
+		bResult &= PakWriterContext.Flush();
 
 		for (const FString& TempOutputDirectory : TempOutputDirectoriesToDelete)
 		{
