UnrealPak: Restore display logging of mount point for pak list command (as before 24117138).
Add the filename for additional context.

#rnx

diff --git a/Engine/Source/Developer/PakFileUtilities/Private/PakFileUtilities.cpp b/Engine/Source/Developer/PakFileUtilities/Private/PakFileUtilities.cpp
--- a/Engine/Source/Developer/PakFileUtilities/Private/PakFileUtilities.cpp
+++ b/Engine/Source/Developer/PakFileUtilities/Private/PakFileUtilities.cpp
@@ -3245,7 +3245,8 @@ bool ListFilesInPak(const TCHAR * InPakFilename, int64 SizeFilter, bool bInclude
 
 	if (PakFile.IsValid())
 	{
-		UE_LOG(LogPakFile, Log, TEXT("Mount point %s"), *PakFile.GetMountPoint());
+		UE_LOG(LogPakFile, Display, TEXT("Listing %s with mount point \"%s\""),
+			*FPaths::GetCleanFilename(InPakFilename), *PakFile.GetMountPoint());
 
 		TArray<FPakFile::FPakEntryIterator> Records;
 
