Fix a crash when generating a patch with a pak larger than 2GB.

diff --git a/Engine/Source/Developer/PakFileUtilities/Private/PakFileUtilities.cpp b/Engine/Source/Developer/PakFileUtilities/Private/PakFileUtilities.cpp
--- a/Engine/Source/Developer/PakFileUtilities/Private/PakFileUtilities.cpp
+++ b/Engine/Source/Developer/PakFileUtilities/Private/PakFileUtilities.cpp
@@ -4526,8 +4526,7 @@ bool GenerateHashesFromPak(const TCHAR* InPakFilename, const TCHAR* InDestPakFil
 					if (EntryInfo.IndexDataEquals(Entry))
 					{
 						// TUniquePtr<FArchive> FileHandle(IFileManager::Get().CreateFileWriter(*DestFilename));
-						TArray<uint8> Bytes;
-						FMemoryWriter MemoryFile(Bytes);
+						FLargeMemoryWriter MemoryFile;
 						FArchive* FileHandle = &MemoryFile;
 						// if (FileHandle.IsValid())
 						{
@@ -4541,7 +4540,7 @@ bool GenerateHashesFromPak(const TCHAR* InPakFilename, const TCHAR* InDestPakFil
 							}
 
 							UE_LOG(LogPakFile, Display, TEXT("Generated hash for \"%s\""), *FullFilename);
-							GenerateHashForFile(Bytes.GetData(), Bytes.Num(), FileHash);
+							GenerateHashForFile(MemoryFile.GetData(), MemoryFile.TotalSize(), FileHash);
 							FileHash.PatchIndex = PakPriority;
 							FileHash.bIsDeleteRecord = false;
 							FileHash.bForceInclude = false;
