[AutoRTFM] Support transactional safety for FRefCountBase and TRefCountPtr.

When a transaction is active, FRefCountBase now defers refcount releases until commit. Note that the refcounts returned from GetRefCount, AddRef and Release inside of a transaction are not trustworthy and should not be relied upon in transactional code, because releases will not be reflected in the count until the transaction is committed.

The MakeRefCount method has been updated to initialize the refcount of an object in the open; this is not required for correctness, but it is more efficient since it avoids setting up a new Abort task.

diff --git a/Engine/Source/Developer/PakFileUtilities/Private/PakFileUtilities.cpp b/Engine/Source/Developer/PakFileUtilities/Private/PakFileUtilities.cpp
--- a/Engine/Source/Developer/PakFileUtilities/Private/PakFileUtilities.cpp
+++ b/Engine/Source/Developer/PakFileUtilities/Private/PakFileUtilities.cpp
@@ -3252,7 +3252,7 @@ void WritePakFooter(FArchive& PakHandle, FPakFooterInfo& Footer)
 
 bool TestPakFile(const TCHAR* Filename, bool TestHashes)
 {	
-	TRefCountPtr<FPakFile> PakFile = new FPakFile(&FPlatformFileManager::Get().GetPlatformFile(), Filename, false);
+	TRefCountPtr<FPakFile> PakFile = MakeRefCount<FPakFile>(&FPlatformFileManager::Get().GetPlatformFile(), Filename, /*bIsSigned=*/ false);
 	if (PakFile->IsValid())
 	{
 		return TestHashes ? PakFile->Check() : true;
@@ -3268,7 +3268,7 @@ bool ListFilesInPak(const TCHAR * InPakFilename, int64 SizeFilter, bool bInclude
 {
 	TRACE_CPUPROFILER_EVENT_SCOPE(ListFilesInPak);
 	IPlatformFile* LowerLevelPlatformFile = &FPlatformFileManager::Get().GetPlatformFile();
-	TRefCountPtr<FPakFile> PakFilePtr = new FPakFile(LowerLevelPlatformFile, InPakFilename, false);
+	TRefCountPtr<FPakFile> PakFilePtr = MakeRefCount<FPakFile>(LowerLevelPlatformFile, InPakFilename, /*bIsSigned=*/ false);
 	FPakFile& PakFile = *PakFilePtr;
 	int32 FileCount = 0;
 	int64 FileSize = 0;
@@ -3527,7 +3527,7 @@ bool AuditPakFiles( const FString& InputPath, bool bOnlyDeleted, const FString&
 		int32 PakPriority = GetPakPriorityFromFilename(PakFilename);
 		HighestPakPriority = FMath::Max( HighestPakPriority, PakPriority );
 
-		TRefCountPtr<FPakFile> PakFilePtr = new FPakFile(&FPlatformFileManager::Get().GetPlatformFile(), *PakFilename, false);
+		TRefCountPtr<FPakFile> PakFilePtr = MakeRefCount<FPakFile>(&FPlatformFileManager::Get().GetPlatformFile(), *PakFilename, /*bIsSigned=*/ false);
 		FPakFile& PakFile = *PakFilePtr;
 		if (PakFile.IsValid())
 		{
@@ -3789,7 +3789,7 @@ bool ListFilesAtOffset( const TCHAR* InPakFileName, const TArray<int64>& InOffse
 		return false;
 	}
 
-	TRefCountPtr<FPakFile> PakFilePtr = new FPakFile(&FPlatformFileManager::Get().GetPlatformFile(), InPakFileName, false);
+	TRefCountPtr<FPakFile> PakFilePtr = MakeRefCount<FPakFile>(&FPlatformFileManager::Get().GetPlatformFile(), InPakFileName, /*bIsSigned=*/ false);
 	FPakFile& PakFile = *PakFilePtr;
 	if (!PakFile.IsValid())
 	{
@@ -3838,7 +3838,7 @@ bool ListFilesAtOffset( const TCHAR* InPakFileName, const TArray<int64>& InOffse
 bool ShowCompressionBlockCRCs( const TCHAR* InPakFileName, TArray<int64>& InOffsets, const FKeyChain& InKeyChain )
 {
 	// open the pak file
-	TRefCountPtr<FPakFile> PakFilePtr = new FPakFile(&FPlatformFileManager::Get().GetPlatformFile(), InPakFileName, false);
+	TRefCountPtr<FPakFile> PakFilePtr = MakeRefCount<FPakFile>(&FPlatformFileManager::Get().GetPlatformFile(), InPakFileName, /*bIsSigned=*/ false);
 	FPakFile& PakFile = *PakFilePtr;
 	if (!PakFile.IsValid())
 	{
@@ -3974,7 +3974,7 @@ bool GeneratePIXMappingFile(const TArray<FString> InPakFileList, const FString&
 			}
 		}
 
-		TRefCountPtr<FPakFile> PakFilePtr = new FPakFile(&FPlatformFileManager::Get().GetPlatformFile(), *PakFileName, false);
+		TRefCountPtr<FPakFile> PakFilePtr = MakeRefCount<FPakFile>(&FPlatformFileManager::Get().GetPlatformFile(), *PakFileName, /*bIsSigned=*/ false);
 		FPakFile& PakFile = *PakFilePtr;
 		if (!PakFile.IsValid())
 		{
@@ -4050,7 +4050,7 @@ bool ExtractFilesFromPak(const TCHAR* InPakFilename, TMap<FString, FFileInfo>& I
 		FString PakFilename = PakFileDirectory + "\\" + PakFileList[PakFileIndex];
 		int32 PakPriority = GetPakPriorityFromFilename(PakFilename);
 
-		TRefCountPtr<FPakFile> PakFilePtr = new FPakFile(&FPlatformFileManager::Get().GetPlatformFile(), *PakFilename, false);
+		TRefCountPtr<FPakFile> PakFilePtr = MakeRefCount<FPakFile>(&FPlatformFileManager::Get().GetPlatformFile(), *PakFilename, /*bIsSigned=*/ false);
 		FPakFile& PakFile = *PakFilePtr;
 		if (PakFile.IsValid())
 		{
@@ -4227,7 +4227,7 @@ void CreateDiffRelativePathMap(TArray<FString>& FileNames, const FString& RootPa
 
 bool DumpPakInfo(const FString& InPakFilename, const FKeyChain& InKeyChain)
 {
-	TRefCountPtr<FPakFile> PakFilePtr = new FPakFile(&FPlatformFileManager::Get().GetPlatformFile(), *InPakFilename, false);
+	TRefCountPtr<FPakFile> PakFilePtr = MakeRefCount<FPakFile>(&FPlatformFileManager::Get().GetPlatformFile(), *InPakFilename, /*bIsSigned=*/ false);
 	FPakFile& PakFile = *PakFilePtr;
 
 	if (!PakFile.IsValid())
@@ -4270,9 +4270,9 @@ bool DiffFilesInPaks(const FString& InPakFilename1, const FString& InPakFilename
 
 	TGuardValue<ELogTimes::Type> DisableLogTimes(GPrintLogTimes, ELogTimes::None);
 
-	TRefCountPtr<FPakFile> PakFilePtr1 = new FPakFile(&FPlatformFileManager::Get().GetPlatformFile(), *InPakFilename1, false);
+	TRefCountPtr<FPakFile> PakFilePtr1 = MakeRefCount<FPakFile>(&FPlatformFileManager::Get().GetPlatformFile(), *InPakFilename1, /*bIsSigned=*/ false);
 	FPakFile& PakFile1 = *PakFilePtr1;
-	TRefCountPtr<FPakFile> PakFilePtr2 = new FPakFile(&FPlatformFileManager::Get().GetPlatformFile(), *InPakFilename2, false);
+	TRefCountPtr<FPakFile> PakFilePtr2 = MakeRefCount<FPakFile>(&FPlatformFileManager::Get().GetPlatformFile(), *InPakFilename2, /*bIsSigned=*/ false);
 	FPakFile& PakFile2 = *PakFilePtr2;
 	if (PakFile1.IsValid() && PakFile2.IsValid())
 	{		
@@ -4467,7 +4467,7 @@ bool GenerateHashesFromPak(const TCHAR* InPakFilename, const TCHAR* InDestPakFil
 		int32 PakPriority = GetPakPriorityFromFilename(PakFilename);
 		int32 PakChunkIndex = GetPakChunkIndexFromFilename(PakFilename);
 
-		TRefCountPtr<FPakFile> PakFilePtr = new FPakFile(&FPlatformFileManager::Get().GetPlatformFile(), *PakFilename, false);
+		TRefCountPtr<FPakFile> PakFilePtr = MakeRefCount<FPakFile>(&FPlatformFileManager::Get().GetPlatformFile(), *PakFilename, /*bIsSigned=*/ false);
 		FPakFile& PakFile = *PakFilePtr;
 		if (PakFile.IsValid())
 		{
@@ -5109,7 +5109,7 @@ void ProcessLegacyFileMoves( TArray<FPakInputPair>& InDeleteRecords, TMap<FStrin
 		UE_LOG(LogPakFile, Display, TEXT("Checking old pak file \"%s\" Pri:%d Chunk:%d."), *PakFilename, PakPriority, PakChunkIndex );
 
 
-		TRefCountPtr<FPakFile> PakFilePtr = new FPakFile(&FPlatformFileManager::Get().GetPlatformFile(), *PakFilename, false);
+		TRefCountPtr<FPakFile> PakFilePtr = MakeRefCount<FPakFile>(&FPlatformFileManager::Get().GetPlatformFile(), *PakFilename, /*bIsSigned=*/ false);
 		FPakFile& PakFile = *PakFilePtr;
 		if (PakFile.IsValid())
 		{
diff --git a/Engine/Source/Editor/UnrealEd/Private/Commandlets/ExportPakDependencies.cpp b/Engine/Source/Editor/UnrealEd/Private/Commandlets/ExportPakDependencies.cpp
--- a/Engine/Source/Editor/UnrealEd/Private/Commandlets/ExportPakDependencies.cpp
+++ b/Engine/Source/Editor/UnrealEd/Private/Commandlets/ExportPakDependencies.cpp
@@ -138,7 +138,7 @@ bool ExportDependencies(const TCHAR * PakFilename, const TCHAR* GameName, const
 	// Example command line used for this tool
 	// C:\Development\BB\WEX\Saved\StagedBuilds\WindowsClient\WorldExplorers\Content\Paks\WorldExplorers-WindowsNoEditor.pak WorldExplorers WEX -exportdependencies=c:\dvtemp\output -debug -NoAssetRegistryCache -ForceDependsGathering
 	
-	TRefCountPtr<FPakFile> PakFilePtr = new FPakFile(&FPlatformFileManager::Get().GetPlatformFile(), PakFilename, bSigned);
+	TRefCountPtr<FPakFile> PakFilePtr = MakeRefCount<FPakFile>(&FPlatformFileManager::Get().GetPlatformFile(), PakFilename, bSigned);
 	FPakFile& PakFile = *PakFilePtr;
 
 	if(PakFile.IsValid())
diff --git a/Engine/Source/Editor/UnrealEd/Private/Factories/PackFactory.cpp b/Engine/Source/Editor/UnrealEd/Private/Factories/PackFactory.cpp
--- a/Engine/Source/Editor/UnrealEd/Private/Factories/PackFactory.cpp
+++ b/Engine/Source/Editor/UnrealEd/Private/Factories/PackFactory.cpp
@@ -391,7 +391,7 @@ UObject* UPackFactory::FactoryCreateBinary
 )
 { 
 	FBufferReader PakReader((void*)Buffer, BufferEnd-Buffer, false);
-	TRefCountPtr<FPakFile> PakFilePtr = new FPakFile(&PakReader);
+	TRefCountPtr<FPakFile> PakFilePtr = MakeRefCount<FPakFile>(&PakReader);
 	FPakFile& PakFile = *PakFilePtr;
 
 	UObject* ReturnAsset = nullptr;
diff --git a/Engine/Source/Programs/AutoRTFMTests/Private/RefCountPtrTests.cpp b/Engine/Source/Programs/AutoRTFMTests/Private/RefCountPtrTests.cpp
new file mode 100644
--- /dev/null
+++ b/Engine/Source/Programs/AutoRTFMTests/Private/RefCountPtrTests.cpp
@@ -0,0 +1,198 @@
+// Copyright Epic Games, Inc. All Rights Reserved.
+
+#include "Catch2Includes.h"
+#include <AutoRTFM/AutoRTFM.h>
+#include "Templates/RefCounting.h"
+
+namespace
+{
+	struct RefCountedInt : public FRefCountBase
+	{
+		RefCountedInt(int InValue) : Value(InValue) {}
+		int Value;
+	};
+}
+
+TEST_CASE("RefCountPtr.PreviouslyAllocated")
+{
+	SECTION("Using new T()")
+	{
+		TRefCountPtr<RefCountedInt> Foo(new RefCountedInt(42));
+
+		AutoRTFM::Commit([&]
+		{
+			// Make a copy to bump the reference count.
+			TRefCountPtr<RefCountedInt> Copy = Foo;
+			Copy->Value = 13;
+		});
+
+		REQUIRE(Foo->Value == 13);
+		REQUIRE(Foo.GetRefCount() == 1);
+	}
+
+	SECTION("Using MakeRefCount<T>")
+	{
+		TRefCountPtr<RefCountedInt> Foo = MakeRefCount<RefCountedInt>(42);
+
+		AutoRTFM::Commit([&]
+		{
+			// Make a copy to bump the reference count.
+			TRefCountPtr<RefCountedInt> Copy = Foo;
+			Copy->Value = 13;
+		});
+
+		REQUIRE(Foo->Value == 13);
+		REQUIRE(Foo.GetRefCount() == 1);
+	}
+}
+
+TEST_CASE("RefCountPtr.AbortWithPreviouslyAllocated")
+{
+	SECTION("Using new T()")
+	{
+		TRefCountPtr<RefCountedInt> Foo(new RefCountedInt(42));
+
+		REQUIRE(AutoRTFM::ETransactionResult::AbortedByRequest == AutoRTFM::Transact([&]
+		{
+			// Make a copy to bump the reference count.
+			TRefCountPtr<RefCountedInt> Copy = Foo;
+			Copy->Value = 13;
+
+			AutoRTFM::AbortTransaction();
+		}));
+
+		REQUIRE(Foo->Value == 42);
+		REQUIRE(Foo.GetRefCount() == 1);
+	}
+
+	SECTION("Using MakeRefCount<T>")
+	{
+		TRefCountPtr<RefCountedInt> Foo = MakeRefCount<RefCountedInt>(42);
+
+		REQUIRE(AutoRTFM::ETransactionResult::AbortedByRequest == AutoRTFM::Transact([&]
+		{
+			// Make a copy to bump the reference count.
+			TRefCountPtr<RefCountedInt> Copy = Foo;
+			Copy->Value = 13;
+
+			AutoRTFM::AbortTransaction();
+		}));
+
+		REQUIRE(Foo->Value == 42);
+		REQUIRE(Foo.GetRefCount() == 1);
+	}
+}
+
+TEST_CASE("RefCountPtr.AbortWithNewlyAllocated")
+{
+	int Result = 42;
+
+	SECTION("Using new T()")
+	{
+		REQUIRE(AutoRTFM::ETransactionResult::AbortedByRequest == AutoRTFM::Transact([&]
+		{
+			TRefCountPtr<RefCountedInt> Foo(new RefCountedInt(42));
+			Result = Foo->Value;
+			AutoRTFM::AbortTransaction();
+		}));
+	}
+
+	SECTION("Using new T() and a copy")
+	{
+		REQUIRE(AutoRTFM::ETransactionResult::AbortedByRequest == AutoRTFM::Transact([&]
+		{
+			TRefCountPtr<RefCountedInt> Foo(new RefCountedInt(42));
+			TRefCountPtr<RefCountedInt> Copy = Foo;
+			Result = Copy->Value;
+			AutoRTFM::AbortTransaction();
+		}));
+	}
+
+	SECTION("Using MakeRefCount<T>")
+	{
+		REQUIRE(AutoRTFM::ETransactionResult::AbortedByRequest == AutoRTFM::Transact([&]
+		{
+			TRefCountPtr<RefCountedInt> Foo = MakeRefCount<RefCountedInt>(42);
+			Result = Foo->Value;
+			AutoRTFM::AbortTransaction();
+		}));
+	}
+
+	SECTION("Using MakeRefCount<T> and a copy")
+	{
+		REQUIRE(AutoRTFM::ETransactionResult::AbortedByRequest == AutoRTFM::Transact([&]
+		{
+			TRefCountPtr<RefCountedInt> Foo = MakeRefCount<RefCountedInt>(42);
+			TRefCountPtr<RefCountedInt> Copy = Foo;
+			Result = Copy->Value;
+			AutoRTFM::AbortTransaction();
+		}));
+	}
+
+	REQUIRE(Result == 42);
+}
+
+TEST_CASE("RefCountPtr.OnCommitCapturingRefCountPtr")
+{
+	TRefCountPtr<RefCountedInt> Foo = MakeRefCount<RefCountedInt>(42);
+
+	SECTION("Committing")
+	{
+		AutoRTFM::Commit([&]
+		{
+			AutoRTFM::OnCommit([Foo]
+			{
+				Foo->Value = 13;
+			});
+		});
+
+		REQUIRE(Foo->Value == 13);
+	}
+
+	SECTION("Aborting")
+	{
+		REQUIRE(AutoRTFM::ETransactionResult::AbortedByRequest == AutoRTFM::Transact([&]
+		{
+			AutoRTFM::OnCommit([Foo]
+			{
+				Foo->Value = 13;
+			});
+			AutoRTFM::AbortTransaction();
+		}));
+
+		REQUIRE(Foo->Value == 42);
+	}
+}
+
+TEST_CASE("RefCountPtr.OnAbortCapturingRefCountPtr")
+{
+	TRefCountPtr<RefCountedInt> Foo = MakeRefCount<RefCountedInt>(42);
+
+	SECTION("Committing")
+	{
+		AutoRTFM::Commit([&]
+		{
+			AutoRTFM::OnAbort([Foo]
+			{
+				Foo->Value = 13;
+			});
+		});
+
+		// The test harness can decide to abort transactions if `ShouldRetryNonNestedTransactions` is set.
+		REQUIRE((Foo->Value == 42 || AutoRTFM::ForTheRuntime::ShouldRetryNonNestedTransactions()));
+	}
+
+	SECTION("Aborting")
+	{
+		REQUIRE(AutoRTFM::ETransactionResult::AbortedByRequest == AutoRTFM::Transact([&]
+		{
+			AutoRTFM::OnAbort([Foo]
+			{
+				Foo->Value = 13;
+			});
+			AutoRTFM::AbortTransaction();
+		}));
+
+		REQUIRE(Foo->Value == 13);
+	}
+}
\ No newline at end of file
diff --git a/Engine/Source/Runtime/Core/Private/Templates/RefCounting.cpp b/Engine/Source/Runtime/Core/Private/Templates/RefCounting.cpp
--- a/Engine/Source/Runtime/Core/Private/Templates/RefCounting.cpp
+++ b/Engine/Source/Runtime/Core/Private/Templates/RefCounting.cpp
@@ -2,9 +2,7 @@
 
 #include "Templates/RefCounting.h"
 
-void 
-FRefCountBase::CheckRefCount() const
+void FRefCountBase::CheckRefCountIsNonZero() const
 {
-	// Verify that Release() is not being called on an object which is already at a zero refcount
-	check(NumRefs != 0);
-}
+	check(!"Release() was called on an object which is already at a zero refcount.");
+}
\ No newline at end of file
diff --git a/Engine/Source/Runtime/Core/Public/Templates/RefCounting.h b/Engine/Source/Runtime/Core/Public/Templates/RefCounting.h
--- a/Engine/Source/Runtime/Core/Public/Templates/RefCounting.h
+++ b/Engine/Source/Runtime/Core/Public/Templates/RefCounting.h
@@ -11,8 +11,9 @@
 #include "Misc/Build.h"
 #include "Serialization/Archive.h"
 #include "Serialization/MemoryLayout.h"
-#include "Templates/TypeHash.h"
 #include "Templates/Requires.h"
+#include "Templates/TypeHash.h"
+#include "Templates/UnrealTemplate.h"
 #include <atomic>
 #include <type_traits>
 
@@ -32,44 +33,121 @@ public:
 class FRefCountBase
 {
 public:
-			FRefCountBase() = default;
+	FRefCountBase() = default;
 	virtual ~FRefCountBase() = default;
 
 	FRefCountBase(const FRefCountBase& Rhs) = delete;
 	FRefCountBase& operator=(const FRefCountBase& Rhs) = delete;
 
-	inline uint32 AddRef() const
+	uint32 AddRef() const
 	{
-		return uint32(FPlatformAtomics::InterlockedIncrement(&NumRefs));
+		uint32 Refs;
+		UE_AUTORTFM_OPEN
+		{
+			Refs = NumRefs++;
+		};
+
+		// If we are inside a transaction, a rollback must undo our refcount bump.
+		// In general, this is best handled by Release(). However, there is one case 
+		// that needs to be handled with special care. A brand-new object has a
+		// refcount of zero, and a rollback must return it to this zero-refcount state.
+		// However, calling AddRef() followed by Release() would not accomplish this;
+		// instead, it would free the object entirely! We need to guard against this,
+		// since it could lead to a double-free, so we detect the zero-reference case
+		// and special-case it.
+		if (Refs == 0)
+		{
+			UE_AUTORTFM_ONABORT(this)
+			{
+				--NumRefs;
+				// The refcount is likely zero now, but leaving the object alive isn't a leak.
+				// We are restoring the object back to its initial "just-created" state.
+			};
+		}
+		else
+		{
+			UE_AUTORTFM_ONABORT(this)
+			{
+				Release();
+			};
+		}
+
+		return Refs + 1;
 	}
 
-	inline uint32 Release() const
+	uint32 Release() const
 	{
 #if DO_GUARD_SLOW
 		if (NumRefs == 0)
 		{
-			CheckRefCount();
+			CheckRefCountIsNonZero();
 		}
 #endif
 
-		const int32 Refs = FPlatformAtomics::InterlockedDecrement(&NumRefs);
-		if (Refs == 0)
+		uint32 RefsToReturn;
+
+		if (AutoRTFM::IsClosed())
 		{
-			delete this;
+			// We return the active number of references minus one to maintain the existing
+			// Release() behavior as closely as possible while inside a transaction, even 
+			// though we are deferring reference count changes until commit time.
+			// Be advised that GetRefCount() would reveal our trickery, since it
+			// always returns the true refcount.
+			UE_AUTORTFM_OPEN
+			{
+				RefsToReturn = uint32(NumRefs.load(std::memory_order_relaxed)) - 1;
+			};
+
+			// Refcount changes and frees are deferred until the transaction is concluded.
+			UE_AUTORTFM_ONCOMMIT(this)
+			{
+				ImmediatelyRelease();
+			};
+		}
+		else
+		{
+			RefsToReturn = ImmediatelyRelease() - 1;
 		}
 
-		return uint32(Refs);
+		return RefsToReturn;
 	}
 
 	uint32 GetRefCount() const
 	{
-		return uint32(NumRefs);
+		// Inside of an AutoRTFM transaction, the returned refcount value may be higher than you'd
+		// expect, because all Release calls are deferred until the transaction commit time.
+		uint32 Refs;
+		UE_AUTORTFM_OPEN
+		{
+			Refs = uint32(NumRefs.load(std::memory_order_relaxed));
+		};
+		return Refs;
 	}
 
 private:
-	mutable int32 NumRefs = 0;
+	uint32 ImmediatelyRelease() const
+	{
+		// This decrements the reference count and immediately frees the object if it has become zero.
+		// A check-failure is issued if an object is over-released.
+		// Returns the refcount _before_ it was decremented.
+		int32 RefsBeforeRelease = NumRefs--;
+#if DO_GUARD_SLOW
+		if (RefsBeforeRelease == 0)
+		{
+			CheckRefCountIsNonZero();
+		}
+#endif
+		if (RefsBeforeRelease == 1)
+		{
+			delete this;
+		}
+		return (uint32)RefsBeforeRelease;
+	}
+
 
-	CORE_API void CheckRefCount() const;
+	mutable std::atomic<int32> NumRefs = 0;
+
+	CORE_API void CheckRefCountIsNonZero() const;
 };
 
 /**
@@ -283,18 +361,15 @@ private:
 template<typename ReferencedType>
 class TRefCountPtr
 {
-	typedef ReferencedType* ReferenceType;
+	using ReferenceType = ReferencedType*;
 
 public:
+	FORCEINLINE TRefCountPtr() = default;
 
-	FORCEINLINE TRefCountPtr():
-		Reference(nullptr)
-	{ }
-
-	TRefCountPtr(ReferencedType* InReference,bool bAddRef = true)
+	TRefCountPtr(ReferencedType* InReference, bool bAddRef = true)
 	{
 		Reference = InReference;
-		if(Reference && bAddRef)
+		if (Reference && bAddRef)
 		{
 			Reference->AddRef();
 		}
@@ -303,7 +378,7 @@ public:
 	TRefCountPtr(const TRefCountPtr& Copy)
 	{
 		Reference = Copy.Reference;
-		if(Reference)
+		if (Reference)
 		{
 			Reference->AddRef();
 		}
@@ -334,7 +409,7 @@ public:
 
 	~TRefCountPtr()
 	{
-		if(Reference)
+		if (Reference)
 		{
 			Reference->Release();
 		}
@@ -377,7 +452,7 @@ public:
 			ReferencedType* OldReference = Reference;
 			Reference = InPtr.Reference;
 			InPtr.Reference = nullptr;
-			if(OldReference)
+			if (OldReference)
 			{
 				OldReference->Release();
 			}
@@ -464,8 +539,7 @@ public:
 	}
 
 private:
-
-	ReferencedType* Reference;
+	ReferencedType* Reference = nullptr;
 
 	template <typename OtherType>
 	friend class TRefCountPtr;
@@ -511,7 +585,15 @@ template <
 	typename... TArgs
 	UE_REQUIRES(!std::is_array_v<T>)
 >
-FORCEINLINE TRefCountPtr<T> MakeRefCount(TArgs&&... Args)
+[[nodiscard]] FORCEINLINE TRefCountPtr<T> MakeRefCount(TArgs&&... Args)
 {
-	return TRefCountPtr<T>(new T(Forward<TArgs>(Args)...));
-}
\ No newline at end of file
+	T* NewUnwrappedObject = new T(Forward<TArgs>(Args)...);
+
+	// Set up `NewObject` in the open to avoid unnecessary (but harmless) OnAbort tasks.
+	TRefCountPtr<T> NewObject;
+	UE_AUTORTFM_OPEN
+	{
+		NewObject = NewUnwrappedObject;
+	};
+	return NewObject;
+}
diff --git a/Engine/Source/Runtime/PakFile/Private/IPlatformFilePak.cpp b/Engine/Source/Runtime/PakFile/Private/IPlatformFilePak.cpp
--- a/Engine/Source/Runtime/PakFile/Private/IPlatformFilePak.cpp
+++ b/Engine/Source/Runtime/PakFile/Private/IPlatformFilePak.cpp
@@ -5317,7 +5317,7 @@ bool FPakPlatformFile::FindFileInPakFiles(TArray<FPakListEntry>& Paks, const TCH
 	return false;
 }
 
-UE_AUTORTFM_ALWAYS_OPEN bool FPakPlatformFile::FindFileInPakFiles(const TCHAR* Filename, TRefCountPtr<FPakFile>* OutPakFile,
+bool FPakPlatformFile::FindFileInPakFiles(const TCHAR* Filename, TRefCountPtr<FPakFile>* OutPakFile,
 	FPakEntry* OutEntry)
 {
 	TArray<FPakListEntry> Paks;
@@ -5802,7 +5802,7 @@ bool FPakPlatformFile::Mount(const TCHAR* InPakFilename, uint32 PakOrder, const
 	bool bIoStoreSuccess = true;
 	if (LowerLevel->FileExists(InPakFilename))
 	{
-		TRefCountPtr<FPakFile> Pak = new FPakFile(LowerLevel, InPakFilename, bSigned, bLoadIndex);
+		TRefCountPtr<FPakFile> Pak = MakeRefCount<FPakFile>(LowerLevel, InPakFilename, bSigned, bLoadIndex);
 		if (Pak.GetReference()->IsValid())
 		{
 			if (!Pak->GetInfo().EncryptionKeyGuid.IsValid() || UE::FEncryptionKeyManager::Get().ContainsKey(Pak->GetInfo().EncryptionKeyGuid))
