Fixed some incorrect printf format specifiers.

diff --git a/Engine/Source/Developer/AssetTools/Private/AssetHeaderPatcher.cpp b/Engine/Source/Developer/AssetTools/Private/AssetHeaderPatcher.cpp
--- a/Engine/Source/Developer/AssetTools/Private/AssetHeaderPatcher.cpp
+++ b/Engine/Source/Developer/AssetTools/Private/AssetHeaderPatcher.cpp
@@ -2958,7 +2958,7 @@ FAssetHeaderPatcher::EResult FAssetHeaderPatcherInner::PatchHeader_WriteDestinat
 			PatchSummaryOffsets(Summary, ImportTableStartOffset, Delta);
 			Summary.TotalHeaderSize += (int32)Delta;
 			checkf(Summary.ImportCount == ImportTable.Num(), TEXT("%d == %d"), Summary.ImportCount, ImportTable.Num());
-			checkf(Summary.ImportOffset == ImportTableStartOffset, TEXT("%d == %d"), Summary.ImportOffset, ImportTableStartOffset);
+			checkf(Summary.ImportOffset == ImportTableStartOffset, TEXT("%d == %" INT64_FMT), Summary.ImportOffset, ImportTableStartOffset);
 
 			break;
 		}
@@ -2979,7 +2979,7 @@ FAssetHeaderPatcher::EResult FAssetHeaderPatcherInner::PatchHeader_WriteDestinat
 			check(Delta == 0);
 			checkf(ExportTableSize == SourceSection.Size, TEXT("%d == %d"), (int)ExportTableSize, (int)SourceSection.Size); // We only patch export table offsets, we should not be patching size
 			checkf(Summary.ExportCount == ExportTable.Num(), TEXT("%d == %d"), Summary.ExportCount, ExportTable.Num());
-			checkf(Summary.ExportOffset == ExportTableStartOffset, TEXT("%d == %d"), Summary.ExportOffset, ExportTableStartOffset);
+			checkf(Summary.ExportOffset == ExportTableStartOffset, TEXT("%d == %" INT64_FMT), Summary.ExportOffset, ExportTableStartOffset);
 
 			break;
 		}
diff --git a/Engine/Source/Developer/DerivedDataCache/Private/HttpCacheStore.cpp b/Engine/Source/Developer/DerivedDataCache/Private/HttpCacheStore.cpp
--- a/Engine/Source/Developer/DerivedDataCache/Private/HttpCacheStore.cpp
+++ b/Engine/Source/Developer/DerivedDataCache/Private/HttpCacheStore.cpp
@@ -2183,7 +2183,7 @@ void FHttpCacheStore::FExistsBatchOp::EndExists(FHttpOperation& Operation)
 	if (Results.Num() != Requests.Num())
 	{
 		UE_LOG(LogDerivedDataCache, Log,
-			TEXT("%s: Cache exists returned unexpected quantity of results (expected %d, got %d)."),
+			TEXT("%s: Cache exists returned unexpected quantity of results (expected %d, got %" UINT64_FMT ")."),
 			*CacheStore.NodeName, Requests.Num(), Results.Num());
 		for (const FCacheGetValueRequest& Request : Requests)
 		{
diff --git a/Engine/Source/Developer/DerivedDataCache/Private/MemoryCacheStore.cpp b/Engine/Source/Developer/DerivedDataCache/Private/MemoryCacheStore.cpp
--- a/Engine/Source/Developer/DerivedDataCache/Private/MemoryCacheStore.cpp
+++ b/Engine/Source/Developer/DerivedDataCache/Private/MemoryCacheStore.cpp
@@ -138,7 +138,7 @@ FMemoryCacheStore::FMemoryCacheStore(
 	FParse::Value(InConfig, TEXT("MaxCacheSize="), SignedMaxCacheSize);
 	const int64 MaxSupportedCacheSize = 2048; // 2 GB
 	MaxCacheSize = FMath::Min(SignedMaxCacheSize, MaxSupportedCacheSize);
-	UE_LOG(LogDerivedDataCache, Display, TEXT("%s: Max Cache Size: %d MB"), *Name, SignedMaxCacheSize);
+	UE_LOG(LogDerivedDataCache, Display, TEXT("%s: Max Cache Size: %" INT64_FMT " MB"), *Name, SignedMaxCacheSize);
 	MaxCacheSize = SignedMaxCacheSize < 0 ? 0 : SignedMaxCacheSize * 1024 * 1024;
 
 	if (StoreOwner)
diff --git a/Engine/Source/Developer/PakFileUtilities/Private/PakFileUtilities.cpp b/Engine/Source/Developer/PakFileUtilities/Private/PakFileUtilities.cpp
--- a/Engine/Source/Developer/PakFileUtilities/Private/PakFileUtilities.cpp
+++ b/Engine/Source/Developer/PakFileUtilities/Private/PakFileUtilities.cpp
@@ -2730,7 +2730,7 @@ void FPakWriterContext::WriterThreadFunc()
 					const uint64 ContiguousGroupedFilePatchWarningThreshhold = 50 * 1024 * 1024;
 					if (OutputEntry->PakFile->ContiguousTotalSizeSmallerThanBlockSize > ContiguousGroupedFilePatchWarningThreshhold)
 					{
-						UE_LOG(LogPakFile, Display, TEXT("%i small files (%i) totaling %llu contiguous bytes found before first 'large' file.  Changes to any of these files could cause the whole group to be 'dirty' in a per-file binary diff based patching system."), OutputEntry->PakFile->ContiguousFilesSmallerThanBlockSize, RequiredPatchPadding, OutputEntry->PakFile->ContiguousTotalSizeSmallerThanBlockSize);
+						UE_LOG(LogPakFile, Display, TEXT("%" UINT64_FMT " small files (%" INT64_FMT ") totaling %" UINT64_FMT " contiguous bytes found before first 'large' file.  Changes to any of these files could cause the whole group to be 'dirty' in a per-file binary diff based patching system."), OutputEntry->PakFile->ContiguousFilesSmallerThanBlockSize, RequiredPatchPadding, OutputEntry->PakFile->ContiguousTotalSizeSmallerThanBlockSize);
 					}
 					OutputEntry->PakFile->ContiguousTotalSizeSmallerThanBlockSize = 0;
 					OutputEntry->PakFile->ContiguousFilesSmallerThanBlockSize = 0;
@@ -2919,7 +2919,7 @@ bool FPakWriterContext::Flush()
 						!EntryFilename.EndsWith(TEXT("uexp")) && // these are export sections of larger files and may be packed with uasset/umap and so we don't need a warning here
 						!(EntryFilename.EndsWith(TEXT(".m.ubulk")) && CmdLineParameters.AlignForMemoryMapping > 0)) // Bulk padding unaligns patch padding and so we don't need a warning here
 					{
-						UE_LOG(LogPakFile, Warning, TEXT("File at offset %lld of size %lld not aligned to patch size %i"), RealStart, PakEntry.Size, RequiredPatchPadding);
+						UE_LOG(LogPakFile, Warning, TEXT("File at offset %" INT64_FMT " of size %" INT64_FMT " not aligned to patch size %" INT64_FMT), RealStart, PakEntry.Size, RequiredPatchPadding);
 					}
 				}
 			}
@@ -2966,31 +2966,31 @@ bool FPakWriterContext::Flush()
 			}
 
 			UE_LOG(LogPakFile, Display, TEXT("Added %d files, %lld bytes total"), OutputPakFile->Index.Num(), TotalPakFileSize);
-			UE_LOG(LogPakFile, Display, TEXT("PrimaryIndex size: %d bytes"), Footer.PrimaryIndexSize);
-			UE_LOG(LogPakFile, Display, TEXT("PathHashIndex size: %d bytes"), Footer.PathHashIndexSize);
-			UE_LOG(LogPakFile, Display, TEXT("FullDirectoryIndex size: %d bytes"), Footer.FullDirectoryIndexSize);
+			UE_LOG(LogPakFile, Display, TEXT("PrimaryIndex size: %" INT64_FMT " bytes"), Footer.PrimaryIndexSize);
+			UE_LOG(LogPakFile, Display, TEXT("PathHashIndex size: %" INT64_FMT " bytes"), Footer.PathHashIndexSize);
+			UE_LOG(LogPakFile, Display, TEXT("FullDirectoryIndex size: %" INT64_FMT " bytes"), Footer.FullDirectoryIndexSize);
 			if (OutputPakFile->TotalUncompressedSize)
 			{
 				float PercentLess = ((float)OutputPakFile->TotalCompressedSize / ((float)OutputPakFile->TotalUncompressedSize / 100.f));
-				UE_LOG(LogPakFile, Display, TEXT("Compression summary: %.2f%% of original size. Compressed Size %lld bytes, Original Size %lld bytes. "), PercentLess, OutputPakFile->TotalCompressedSize, OutputPakFile->TotalUncompressedSize);
+				UE_LOG(LogPakFile, Display, TEXT("Compression summary: %.2f%% of original size. Compressed Size %" UINT64_FMT " bytes, Original Size %" UINT64_FMT " bytes. "), PercentLess, OutputPakFile->TotalCompressedSize, OutputPakFile->TotalUncompressedSize);
 			}
 
 			if (OutputPakFile->RehydratedCount > 0)
 			{
 				UE_LOG(LogPakFile, Display, TEXT("Asset Rehydration"));
 
-				UE_LOG(LogPakFile, Display, TEXT("  Rehydrated: %llu payloads"), OutputPakFile->RehydratedCount);
-				UE_LOG(LogPakFile, Display, TEXT("  Rehydrated: %llu bytes (%.2fMB)"), OutputPakFile->RehydratedBytes, (float)OutputPakFile->RehydratedBytes / 1024.0f / 1024.0f);
+				UE_LOG(LogPakFile, Display, TEXT("  Rehydrated: %" UINT64_FMT " payloads"), OutputPakFile->RehydratedCount);
+				UE_LOG(LogPakFile, Display, TEXT("  Rehydrated: %" UINT64_FMT " bytes (%.2fMB)"), OutputPakFile->RehydratedBytes, (float)OutputPakFile->RehydratedBytes / 1024.0f / 1024.0f);
 			}
 
 			if (OutputPakFile->TotalEncryptedDataSize)
 			{
 				UE_LOG(LogPakFile, Display, TEXT("Encryption - ENABLED"));
-				UE_LOG(LogPakFile, Display, TEXT("  Files: %d"), OutputPakFile->TotalEncryptedFiles);
+				UE_LOG(LogPakFile, Display, TEXT("  Files: %" UINT64_FMT), OutputPakFile->TotalEncryptedFiles);
 
 				if (OutputPakFile->Info.bEncryptedIndex)
 				{
-					UE_LOG(LogPakFile, Display, TEXT("  Index: Encrypted (%d bytes, %.2fMB)"), OutputPakFile->Info.IndexSize, (float)OutputPakFile->Info.IndexSize / 1024.0f / 1024.0f);
+					UE_LOG(LogPakFile, Display, TEXT("  Index: Encrypted (%" INT64_FMT " bytes, %.2fMB)"), OutputPakFile->Info.IndexSize, (float)OutputPakFile->Info.IndexSize / 1024.0f / 1024.0f);
 				}
 				else
 				{
@@ -2998,7 +2998,7 @@ bool FPakWriterContext::Flush()
 				}
 
 
-				UE_LOG(LogPakFile, Display, TEXT("  Total: %d bytes (%.2fMB)"), OutputPakFile->TotalEncryptedDataSize, (float)OutputPakFile->TotalEncryptedDataSize / 1024.0f / 1024.0f);
+				UE_LOG(LogPakFile, Display, TEXT("  Total: %" UINT64_FMT " bytes (%.2fMB)"), OutputPakFile->TotalEncryptedDataSize, (float)OutputPakFile->TotalEncryptedDataSize / 1024.0f / 1024.0f);
 			}
 			else
 			{
@@ -3007,7 +3007,7 @@ bool FPakWriterContext::Flush()
 
 			if (OutputPakFile->TotalEncryptedFiles < OutputPakFile->TotalRequestedEncryptedFiles)
 			{
-				UE_LOG(LogPakFile, Display, TEXT("%d files requested encryption, but no AES key was supplied! Encryption was skipped for these files"), OutputPakFile->TotalRequestedEncryptedFiles);
+				UE_LOG(LogPakFile, Display, TEXT("%" UINT64_FMT " files requested encryption, but no AES key was supplied! Encryption was skipped for these files"), OutputPakFile->TotalRequestedEncryptedFiles);
 			}
 			UE_LOG(LogPakFile, Display, TEXT(""));
 		}
@@ -3019,7 +3019,7 @@ bool FPakWriterContext::Flush()
 
 	if (TotalFilesWithPoorForcedCompression > 0)
 	{
-		UE_LOG(LogPakFile, Display, TEXT("Num files forcibly compressed due to -forcecompress option: %i, using %i bytes extra"), (int64)TotalFilesWithPoorForcedCompression, (int64)TotalExtraMemoryForPoorForcedCompression);
+		UE_LOG(LogPakFile, Display, TEXT("Num files forcibly compressed due to -forcecompress option: %" INT64_FMT ", using %" INT64_FMT "bytes extra"), (int64)TotalFilesWithPoorForcedCompression, (int64)TotalExtraMemoryForPoorForcedCompression);
 	}
 
 #if DETAILED_UNREALPAK_TIMING
@@ -3405,7 +3405,7 @@ bool ListFilesInPak(const TCHAR * InPakFilename, int64 SizeFilter, bool bInclude
 					{
 						if (InspectChunks.Contains(Chunk))
 						{
-							UE_LOG(LogPakFile, Display, TEXT("[%d - %d] \"%s%s\" offset: %lld, size: %d bytes, sha1: %s, compression: %s."),
+							UE_LOG(LogPakFile, Display, TEXT("[%d - %d] \"%s%s\" offset: %" INT64_FMT ", size: %" INT64_FMT " bytes, sha1: %s, compression: %s."),
 								FirstChunk, LastChunk, *MountPoint,
 								It.Key ? **It.Key : TEXT("<FileNamesNotLoaded"), Entry.Offset, Entry.Size,
 								*BytesToHex(Hash, sizeof(FPakEntry::Hash)),
@@ -3416,7 +3416,7 @@ bool ListFilesInPak(const TCHAR * InPakFilename, int64 SizeFilter, bool bInclude
 				}
 				else
 				{
-					UE_LOG(LogPakFile, Display, TEXT("\"%s%s\" offset: %lld, size: %d bytes, sha1: %s, compression: %s."),
+					UE_LOG(LogPakFile, Display, TEXT("\"%s%s\" offset: %" INT64_FMT ", size: %" INT64_FMT " bytes, sha1: %s, compression: %s."),
 						*MountPoint, It.Key ? **It.Key : TEXT("<FileNamesNotLoaded"),
 						Entry.Offset, Entry.Size, *BytesToHex(Hash, sizeof(FPakEntry::Hash)),
 						*PakFile.GetInfo().GetCompressionMethod(Entry.CompressionMethodIndex).ToString());
@@ -3818,7 +3818,7 @@ bool ListFilesAtOffset( const TCHAR* InPakFileName, const TArray<int64>& InOffse
 			if( Offset >= Entry.Offset && Offset <= Entry.Offset+Entry.Size )
 			{
 				const FString* Filename = It.TryGetFilename();
-				UE_LOG( LogPakFile, Display, TEXT("%-12lld%-12lld%-12d%s"), Offset, Entry.Offset, Entry.Size, Filename ? **Filename : TEXT("<FileNamesNotLoaded>") );
+				UE_LOG( LogPakFile, Display, TEXT("%-12" INT64_FMT "%-12" INT64_FMT "%-12" INT64_FMT "%s"), Offset, Entry.Offset, Entry.Size, Filename ? **Filename : TEXT("<FileNamesNotLoaded>") );
 				FoundOffset = Offset;
 				break;
 			}
@@ -3834,7 +3834,7 @@ bool ListFilesAtOffset( const TCHAR* InPakFileName, const TArray<int64>& InOffse
 	//list out any that we didn't find a match for
 	for( int64 InvalidOffset : OffsetsToCheck )
 	{
-		UE_LOG(LogPakFile, Display, TEXT("%-12lld - invalid offset"), InvalidOffset );
+		UE_LOG(LogPakFile, Display, TEXT("%-12" INT64_FMT " - invalid offset"), InvalidOffset );
 	}
 
 	return true;
@@ -3860,7 +3860,7 @@ bool ShowCompressionBlockCRCs( const TCHAR* InPakFileName, TArray<int64>& InOffs
 		//sanity check the offset
 		if (Offset < 0 || Offset > PakFile.TotalSize() )
 		{
-			UE_LOG(LogPakFile, Error, TEXT("Offset: %lld - out of range (max size is %lld)"), Offset, PakFile.TotalSize() );
+			UE_LOG(LogPakFile, Error, TEXT("Offset: %" INT64_FMT " - out of range (max size is %" INT64_FMT ")"), Offset, PakFile.TotalSize() );
 			continue;
 		}
 
@@ -3875,7 +3875,7 @@ bool ShowCompressionBlockCRCs( const TCHAR* InPakFileName, TArray<int64>& InOffs
 				FString EntryFilename = It.TryGetFilename() ? *It.TryGetFilename() : TEXT("<FileNamesNotLoaded>");
 				FName EntryCompressionMethod = PakFile.GetInfo().GetCompressionMethod(Entry->CompressionMethodIndex);
 
-				UE_LOG(LogPakFile, Display, TEXT("Offset: %lld  -> EntrySize: %lld  Encrypted: %-3s  Compression: %-8s  [%s]"), Offset, Entry->Size, Entry->IsEncrypted() ? TEXT("Yes") : TEXT("No"), *EntryCompressionMethod.ToString(), *EntryFilename );
+				UE_LOG(LogPakFile, Display, TEXT("Offset: %" INT64_FMT "  -> EntrySize: %" INT64_FMT "  Encrypted: %-3s  Compression: %-8s  [%s]"), Offset, Entry->Size, Entry->IsEncrypted() ? TEXT("Yes") : TEXT("No"), *EntryCompressionMethod.ToString(), *EntryFilename );
 				break;
 			}
 		}
@@ -3939,7 +3939,7 @@ bool ShowCompressionBlockCRCs( const TCHAR* InPakFileName, TArray<int64>& InOffs
 			uint32 TruncatedCompressedBlockSize = IntCastChecked<uint32>(CompressedBlockSize);
 			const uint32 BlockCrc32 = FCrc::MemCrc32(PersistentBuffer, TruncatedCompressedBlockSize);
 			const FString HexBytes = BytesToHex(PersistentBuffer,FMath::Min(TruncatedCompressedBlockSize,32U));
-			UE_LOG(LogPakFile, Display, TEXT("    Block:%-6d  ProcessedSize: %-6d  DecompressionRawSize: %-6d  Crc32: %-12u [%s...]"), BlockIndex, ProcessedSize, CompressedBlockSize, BlockCrc32, *HexBytes );
+			UE_LOG(LogPakFile, Display, TEXT("    Block:%-6d  ProcessedSize: %-6d  DecompressionRawSize: %-6" INT64_FMT "  Crc32: %-12u [%s...]"), BlockIndex, ProcessedSize, CompressedBlockSize, BlockCrc32, *HexBytes );
 		}
 
 		FMemory::Free(PersistentBuffer);
@@ -4002,7 +4002,7 @@ bool GeneratePIXMappingFile(const TArray<FString> InPakFileList, const FString&
 		{
 			const FPakEntry& Entry = It.Info();
 
-			CSVFileWriter->Logf(TEXT("0x%010llx,0x%08llx,%s"), Entry.Offset, Entry.Size, *(PakFileMountPoint / *It.TryGetFilename()));
+			CSVFileWriter->Logf(TEXT("0x%010" INT64_FMT ",0x%08" INT64_FMT ",%s"), Entry.Offset, Entry.Size, *(PakFileMountPoint / *It.TryGetFilename()));
 		}
 
 		CSVFileWriter->Close();
@@ -4253,8 +4253,8 @@ bool DumpPakInfo(const FString& InPakFilename, const FKeyChain& InKeyChain)
 
 	UE_LOG(LogPakFile, Display, TEXT("Pak File: %s"), *InPakFilename);
 	UE_LOG(LogPakFile, Display, TEXT("    Version: %d"), Info.Version);
-	UE_LOG(LogPakFile, Display, TEXT("    IndexOffset: %d"), Info.IndexOffset);
-	UE_LOG(LogPakFile, Display, TEXT("    IndexSize: %d"), Info.IndexSize);
+	UE_LOG(LogPakFile, Display, TEXT("    IndexOffset: %" INT64_FMT), Info.IndexOffset);
+	UE_LOG(LogPakFile, Display, TEXT("    IndexSize: %" INT64_FMT), Info.IndexSize);
 	UE_LOG(LogPakFile, Display, TEXT("    IndexHash: %s"), *Info.IndexHash.ToString());
 	UE_LOG(LogPakFile, Display, TEXT("    bEncryptedIndex: %d"), Info.bEncryptedIndex);
 	UE_LOG(LogPakFile, Display, TEXT("    EncryptionKeyGuid: %s"), *Info.EncryptionKeyGuid.ToString());
@@ -4325,7 +4325,7 @@ bool DiffFilesInPaks(const FString& InPakFilename1, const FString& InPakFilename
 				++NumUniquePAK1;
 				if (bLogUniques1)
 				{
-					UE_LOG(LogPakFile, Log, TEXT("UniqueToFirstPak, %s, %i, 0"), *PAK1FileName, EntryInfo1.UncompressedSize);
+					UE_LOG(LogPakFile, Log, TEXT("UniqueToFirstPak, %s, %" INT64_FMT ", 0"), *PAK1FileName, EntryInfo1.UncompressedSize);
 				}
 				continue;
 			}
@@ -4343,7 +4343,7 @@ bool DiffFilesInPaks(const FString& InPakFilename1, const FString& InPakFilename
 			//check sizes first as quick compare.
 			if (EntryInfo1.UncompressedSize != EntryInfo2.UncompressedSize)
 			{
-				UE_LOG(LogPakFile, Log, TEXT("FilesizeDifferent, %s, %i, %i"), *PAK1FileName, EntryInfo1.UncompressedSize, EntryInfo2.UncompressedSize);
+				UE_LOG(LogPakFile, Log, TEXT("FilesizeDifferent, %s, %" INT64_FMT ", %" INT64_FMT), *PAK1FileName, EntryInfo1.UncompressedSize, EntryInfo2.UncompressedSize);
 				continue;
 			}
 			
@@ -4373,7 +4373,7 @@ bool DiffFilesInPaks(const FString& InPakFilename1, const FString& InPakFilename
 				if (FMemory::Memcmp(PAKWriter1.GetData(), PAKWriter2.GetData(), EntryInfo1.UncompressedSize) != 0)
 				{
 					++NumDifferentContents;
-					UE_LOG(LogPakFile, Log, TEXT("ContentsDifferent, %s, %i, %i"), *PAK1FileName, EntryInfo1.UncompressedSize, EntryInfo2.UncompressedSize);
+					UE_LOG(LogPakFile, Log, TEXT("ContentsDifferent, %s, %" INT64_FMT ", %" INT64_FMT), *PAK1FileName, EntryInfo1.UncompressedSize, EntryInfo2.UncompressedSize);
 				}
 				else
 				{
@@ -4401,7 +4401,7 @@ bool DiffFilesInPaks(const FString& InPakFilename1, const FString& InPakFilename
 					++NumUniquePAK2;
 					if (bLogUniques2)
 					{
-						UE_LOG(LogPakFile, Log, TEXT("UniqueToSecondPak, %s, 0, %i"), *PAK2FileName, Entry2.UncompressedSize);
+						UE_LOG(LogPakFile, Log, TEXT("UniqueToSecondPak, %s, 0, %" INT64_FMT), *PAK2FileName, Entry2.UncompressedSize);
 					}
 					continue;
 				}
@@ -4607,7 +4607,7 @@ bool FileIsIdentical(FString SourceFile, FString DestFilename, const FFileInfo*
 	if (SourceTotalSize != DestTotalSize)
 	{
 		// file size doesn't match 
-		UE_LOG(LogPakFile, Display, TEXT("Source file size for %s %d bytes doesn't match %s %d bytes, did find %d"), *SourceFile, SourceTotalSize, *DestFilename, DestTotalSize, Hash ? 1 : 0);
+		UE_LOG(LogPakFile, Display, TEXT("Source file size for %s %" INT64_FMT " bytes doesn't match %s %" INT64_FMT " bytes, did find %d"), *SourceFile, SourceTotalSize, *DestFilename, DestTotalSize, Hash ? 1 : 0);
 		return false;
 	}
 
@@ -4821,7 +4821,7 @@ bool DoGapFillingIteration(const TArray<FPakInputPair>& FilesToPak, const TArray
 	}
 	if (PatchSize > MaxPatchSize && !bForce)
 	{
-		UE_LOG(LogPakFile, SEEK_OPT_VERBOSITY, TEXT("Patch seek optimization step %s - Gap size: %dKB - FAILED (patch too big)"), *PrefixString, MaxGapSizeBytes / 1024);
+		UE_LOG(LogPakFile, SEEK_OPT_VERBOSITY, TEXT("Patch seek optimization step %s - Gap size: %" INT64_FMT "KB - FAILED (patch too big)"), *PrefixString, MaxGapSizeBytes / 1024);
 		return false;
 	}
 
@@ -4829,11 +4829,11 @@ bool DoGapFillingIteration(const TArray<FPakInputPair>& FilesToPak, const TArray
 	float FragmentationPercentageAfterGapFill = GetFragmentationPercentage(FilesToPak, IncludeFilesMask, MaxAdjacentOrderDiff, true);
 	if (FragmentationPercentageAfterGapFill >= FragmentationPercentageOriginal && !bForce)
 	{
-		UE_LOG(LogPakFile, SEEK_OPT_VERBOSITY, TEXT("Patch seek optimization step %s - Gap size: %dKB - FAILED (contiguous block count didn't improve)"), *PrefixString, MaxGapSizeBytes / 1024);
+		UE_LOG(LogPakFile, SEEK_OPT_VERBOSITY, TEXT("Patch seek optimization step %s - Gap size: %" INT64_FMT "KB - FAILED (contiguous block count didn't improve)"), *PrefixString, MaxGapSizeBytes / 1024);
 		return false;
 	}
 	InOutIncludeFilesMask = IncludeFilesMask;
-	UE_LOG(LogPakFile, SEEK_OPT_VERBOSITY, TEXT("Patch seek optimization step %s - Gap size: %dKB - SUCCEEDED"), *PrefixString, MaxGapSizeBytes / 1024);
+	UE_LOG(LogPakFile, SEEK_OPT_VERBOSITY, TEXT("Patch seek optimization step %s - Gap size: %" INT64_FMT "KB - SUCCEEDED"), *PrefixString, MaxGapSizeBytes / 1024);
 	return true;
 }
 
diff --git a/Engine/Source/Developer/TextureFormatETC2/Private/TextureFormatETC2.cpp b/Engine/Source/Developer/TextureFormatETC2/Private/TextureFormatETC2.cpp
--- a/Engine/Source/Developer/TextureFormatETC2/Private/TextureFormatETC2.cpp
+++ b/Engine/Source/Developer/TextureFormatETC2/Private/TextureFormatETC2.cpp
@@ -300,7 +300,7 @@ public:
 
 		if (NumSlices * NumBlocksX * NumBlocksY * BytesPerBlock != InEncodedData.GetSize())
 		{
-			UE_LOG(LogTextureFormatETC2, Error, TEXT("Can't decode ETC2 image: incorrect amount of encoded data for image size: %d x %d x %d, expected %llu got %llu"),
+			UE_LOG(LogTextureFormatETC2, Error, TEXT("Can't decode ETC2 image: incorrect amount of encoded data for image size: %d x %d x %" UINT64_FMT ", expected %" UINT64_FMT " got %" UINT64_FMT),
 				InSizeX, InSizeY, NumSlices, NumSlices * NumBlocksX * NumBlocksY * BytesPerBlock, InEncodedData.GetSize());
 			return false;
 		}
diff --git a/Engine/Source/Developer/TraceInsights/Private/Insights/TimingProfiler/Tracks/CpuTimingTrack.cpp b/Engine/Source/Developer/TraceInsights/Private/Insights/TimingProfiler/Tracks/CpuTimingTrack.cpp
--- a/Engine/Source/Developer/TraceInsights/Private/Insights/TimingProfiler/Tracks/CpuTimingTrack.cpp
+++ b/Engine/Source/Developer/TraceInsights/Private/Insights/TimingProfiler/Tracks/CpuTimingTrack.cpp
@@ -29,7 +29,7 @@ INSIGHTS_IMPLEMENT_RTTI(FCpuTimingTrack)
 
 void FCpuTimingTrack::AddTaskInfo(FTooltipDrawState& InOutTooltip, const TraceServices::FTaskInfo& Task) const
 {
-	InOutTooltip.AddTextLine(FString::Printf(TEXT("-------- Task %d%s --------"), Task.Id, Task.bTracked ? TEXT("") : TEXT(" (not tracked)")), FLinearColor::Green);
+	InOutTooltip.AddTextLine(FString::Printf(TEXT("-------- Task %" UINT64_FMT "%s --------"), Task.Id, Task.bTracked ? TEXT("") : TEXT(" (not tracked)")), FLinearColor::Green);
 
 	if (Task.DebugName != nullptr)
 	{
diff --git a/Engine/Source/Developer/TraceServices/Private/AnalysisCache.cpp b/Engine/Source/Developer/TraceServices/Private/AnalysisCache.cpp
--- a/Engine/Source/Developer/TraceServices/Private/AnalysisCache.cpp
+++ b/Engine/Source/Developer/TraceServices/Private/AnalysisCache.cpp
@@ -66,7 +66,7 @@ FAnalysisCache::FFileContents::FFileContents(const TCHAR* FilePath)
 
 			for (const FFileContents::FBlockEntry& Block : Blocks)
 			{
-				UE_LOG(LogAnalysisCache, VeryVerbose, TEXT("   %10d   %10d   %10d kb   %10d kb   %10d kb"), GetCacheId(Block.BlockKey), GetBlockIndex(Block.BlockKey), Block.Offset / 1024, Block.UncompressedSize / 1024, Block.CompressedSize / 1024);
+				UE_LOG(LogAnalysisCache, VeryVerbose, TEXT("   %10d   %10d   %10" UINT64_FMT " kb   %10" UINT64_FMT " kb   %10" UINT64_FMT " kb"), GetCacheId(Block.BlockKey), GetBlockIndex(Block.BlockKey), Block.Offset / 1024, Block.UncompressedSize / 1024, Block.CompressedSize / 1024);
 			}
 
 			//todo: At this point we can check the file size and if it doesn't match we can trim unknown segments.
@@ -390,7 +390,7 @@ uint64 FAnalysisCache::FFileContents::UpdateBlock(FMemoryView Block, BlockKeyTyp
 			File->Seek(Entry.Offset);
 			if(!File->Write((const uint8*)Block.GetData(), Block.GetSize()))
 			{
-				UE_LOG(LogAnalysisCache, Error, TEXT("Failed to update block 0x%x at offset %d kb"), BlockKey, Entry.Offset / 1024);
+				UE_LOG(LogAnalysisCache, Error, TEXT("Failed to update block 0x%x at offset %" UINT64_FMT " kb"), BlockKey, Entry.Offset / 1024);
 				return 0;
 			}
 
@@ -404,10 +404,10 @@ uint64 FAnalysisCache::FFileContents::UpdateBlock(FMemoryView Block, BlockKeyTyp
 		const bool bSeekResult = File->SeekFromEnd(0);
 		check(bSeekResult);
 		const uint64 Offset = File->Tell();
-		checkf(Offset >= ReservedSizeV2, TEXT("Offset (%ull) is larger than reserved size (%u)"), Offset, ReservedSizeV2);
+		checkf(Offset >= ReservedSizeV2, TEXT("Offset (%" UINT64_FMT ") is larger than reserved size (%u)"), Offset, ReservedSizeV2);
 		if(!File->Write((const uint8*) Block.GetData(), Block.GetSize()))
 		{
-			UE_LOG(LogAnalysisCache, Error, TEXT("Failed to update block 0x%x at offset %u kb"), BlockKey, Offset / 1024);
+			UE_LOG(LogAnalysisCache, Error, TEXT("Failed to update block 0x%x at offset %" UINT64_FMT " kb"), BlockKey, Offset / 1024);
 			return 0;
 		}
 
@@ -442,14 +442,14 @@ uint64 FAnalysisCache::FFileContents::LoadBlock(FMutableMemoryView Block, BlockK
 	
 	if (!File->Seek(Entry.Offset))
 	{
-		UE_LOG(LogAnalysisCache, Error, TEXT("Block 0x%x was located on an invalid offset %u kb."), BlockKey, Entry.Offset / 1024);
+		UE_LOG(LogAnalysisCache, Error, TEXT("Block 0x%x was located on an invalid offset %" UINT64_FMT " kb."), BlockKey, Entry.Offset / 1024);
 		return 0;
 	}
 
 	//todo: Add compression
 	if (!File->Read((uint8*)Block.GetData(), Entry.UncompressedSize))
 	{
-		UE_LOG(LogAnalysisCache, Error, TEXT("Unable to read block 0x%x on offset %u kb with size %u kb."), BlockKey, Entry.Offset / 1024, Entry.UncompressedSize / 1024);
+		UE_LOG(LogAnalysisCache, Error, TEXT("Unable to read block 0x%x on offset %" UINT64_FMT " kb with size %" UINT64_FMT " kb."), BlockKey, Entry.Offset / 1024, Entry.UncompressedSize / 1024);
 		return 0;
 	}
 
diff --git a/Engine/Source/Developer/TraceServices/Private/Model/AllocationsProvider.cpp b/Engine/Source/Developer/TraceServices/Private/Model/AllocationsProvider.cpp
--- a/Engine/Source/Developer/TraceServices/Private/Model/AllocationsProvider.cpp
+++ b/Engine/Source/Developer/TraceServices/Private/Model/AllocationsProvider.cpp
@@ -2320,7 +2320,7 @@ void FAllocationsProvider::EditUnmarkAllocationAsHeap(uint32 ThreadId, double Ti
 						int NumAllocsInTop = 10; // only show first 10 allocs
 						for (const FAllocInHeap& AllocInHeap : AllocsInHeap)
 						{
-							UE_LOG(LogTraceServices, Warning, TEXT("[MemAlloc]     alloc 0x%llX (%u bytes)"), AllocInHeap.Address, AllocInHeap.Size);
+							UE_LOG(LogTraceServices, Warning, TEXT("[MemAlloc]     alloc 0x%" UINT64_X_FMT " (%" UINT64_FMT " bytes)"), AllocInHeap.Address, AllocInHeap.Size);
 							if (--NumAllocsInTop <= 0)
 							{
 								break;
diff --git a/Engine/Source/Developer/TraceServices/Private/Model/TasksProfiler.cpp b/Engine/Source/Developer/TraceServices/Private/Model/TasksProfiler.cpp
--- a/Engine/Source/Developer/TraceServices/Private/Model/TasksProfiler.cpp
+++ b/Engine/Source/Developer/TraceServices/Private/Model/TasksProfiler.cpp
@@ -141,14 +141,14 @@ namespace TraceServices
 
 	void FTasksProvider::TaskCreated(TaskTrace::FId TaskId, double Timestamp, uint32 ThreadId, uint64 TaskSize)
 	{
-		UE_LOG(LogTraceServices, Verbose, TEXT("TaskCreated(TaskId: %d, Timestamp %.6f)"), TaskId, Timestamp);
+		UE_LOG(LogTraceServices, Verbose, TEXT("TaskCreated(TaskId: %" UINT64_FMT ", Timestamp %.6f)"), TaskId, Timestamp);
 
 		InitTaskIdToIndexConversion(TaskId);
 
 		FTaskInfo* Task = TryGetOrCreateTask(TaskId);
 		if (Task == nullptr)
 		{
-			UE_LOG(LogTraceServices, Log, TEXT("TaskCreated(TaskId %d, Timestamp %.6f) skipped"), TaskId, Timestamp);
+			UE_LOG(LogTraceServices, Log, TEXT("TaskCreated(TaskId %" UINT64_FMT ", Timestamp %.6f) skipped"), TaskId, Timestamp);
 			return;
 		}
 
@@ -176,14 +176,14 @@ namespace TraceServices
 
 	void FTasksProvider::TaskLaunched(TaskTrace::FId TaskId, const TCHAR* DebugName, bool bTracked, int32 ThreadToExecuteOn, double Timestamp, uint32 ThreadId, uint64 TaskSize)
 	{
-		UE_LOG(LogTraceServices, Verbose, TEXT("TaskLaunched(TaskId: %d, DebugName: %s, bTracked: %d, Timestamp %.6f)"), TaskId, DebugName, bTracked, Timestamp);
+		UE_LOG(LogTraceServices, Verbose, TEXT("TaskLaunched(TaskId: %" UINT64_FMT ", DebugName: %s, bTracked: %d, Timestamp %.6f)"), TaskId, DebugName, bTracked, Timestamp);
 
 		InitTaskIdToIndexConversion(TaskId);
 
 		FTaskInfo* Task = TryGetOrCreateTask(TaskId);
 		if (Task == nullptr)
 		{
-			UE_LOG(LogTraceServices, Log, TEXT("TaskLaunched(TaskId %d, DebugName %s, bTracked %d, Timestamp %.6f) skipped"), TaskId, DebugName, bTracked, Timestamp);
+			UE_LOG(LogTraceServices, Log, TEXT("TaskLaunched(TaskId %" UINT64_FMT ", DebugName %s, bTracked %d, Timestamp %.6f) skipped"), TaskId, DebugName, bTracked, Timestamp);
 			return;
 		}
 
@@ -218,7 +218,7 @@ namespace TraceServices
 		FTaskInfo* Task = TryGetOrCreateTask(TaskId);
 		if (Task == nullptr)
 		{
-			UE_LOG(LogTraceServices, Log, TEXT("TaskScheduled(TaskId %d, Timestamp %.6f) skipped"), TaskId, Timestamp);
+			UE_LOG(LogTraceServices, Log, TEXT("TaskScheduled(TaskId %" UINT64_FMT ", Timestamp %.6f) skipped"), TaskId, Timestamp);
 			return;
 		}
 
@@ -255,14 +255,14 @@ namespace TraceServices
 		FTaskInfo* Task = TryGetOrCreateTask(TaskId);
 		if (Task == nullptr)
 		{
-			UE_LOG(LogTraceServices, Log, TEXT("SubsequentAdded(TaskId %d, SubsequentId %d, Timestamp %.6f) skipped, task %d doesn't exists"), TaskId, SubsequentId, Timestamp, TaskId);
+			UE_LOG(LogTraceServices, Log, TEXT("SubsequentAdded(TaskId %" UINT64_FMT ", SubsequentId %" UINT64_FMT ", Timestamp %.6f) skipped, task %" UINT64_FMT " doesn't exists"), TaskId, SubsequentId, Timestamp, TaskId);
 			return;
 		}
 
 		FTaskInfo* SubsequentTask = TryGetOrCreateTask(SubsequentId);
 		if (SubsequentTask == nullptr)
 		{
-			UE_LOG(LogTraceServices, Log, TEXT("SubsequentAdded(TaskId %d, SubsequentId %d, Timestamp %.6f) skipped, subsequent task %d doesn't exists"), TaskId, SubsequentId, Timestamp, SubsequentId);
+			UE_LOG(LogTraceServices, Log, TEXT("SubsequentAdded(TaskId %" UINT64_FMT ", SubsequentId %" UINT64_FMT ", Timestamp %.6f) skipped, subsequent task %" UINT64_FMT " doesn't exists"), TaskId, SubsequentId, Timestamp, SubsequentId);
 			return;
 		}
 
@@ -289,7 +289,7 @@ namespace TraceServices
 		FTaskInfo* Task = TryGetOrCreateTask(TaskId);
 		if (Task == nullptr)
 		{
-			UE_LOG(LogTraceServices, Log, TEXT("TaskStarted(TaskId %d, Timestamp %.6f) skipped"), TaskId, Timestamp);
+			UE_LOG(LogTraceServices, Log, TEXT("TaskStarted(TaskId %" UINT64_FMT ", Timestamp %.6f) skipped"), TaskId, Timestamp);
 			return;
 		}
 
@@ -334,7 +334,7 @@ namespace TraceServices
 		FTaskInfo* Task = TryGetTask(TaskId);
 		if (Task == nullptr)
 		{
-			UE_LOG(LogTraceServices, Log, TEXT("TaskFinished(TaskId %d, Timestamp %.6f) skipped"), TaskId, Timestamp);
+			UE_LOG(LogTraceServices, Log, TEXT("TaskFinished(TaskId %" UINT64_FMT ", Timestamp %.6f) skipped"), TaskId, Timestamp);
 			return;
 		}
 
@@ -369,7 +369,7 @@ namespace TraceServices
 		FTaskInfo* Task = TryGetOrCreateTask(TaskId);
 		if (Task == nullptr)
 		{
-			UE_LOG(LogTraceServices, Log, TEXT("TaskCompleted(TaskId %d, Timestamp %.6f) skipped"), TaskId, Timestamp);
+			UE_LOG(LogTraceServices, Log, TEXT("TaskCompleted(TaskId %" UINT64_FMT ", Timestamp %.6f) skipped"), TaskId, Timestamp);
 			return;
 		}
 
@@ -390,7 +390,7 @@ namespace TraceServices
 		FTaskInfo* Task = TryGetOrCreateTask(TaskId);
 		if (Task == nullptr)
 		{
-			UE_LOG(LogTraceServices, Log, TEXT("TaskDestroyed(TaskId %d, Timestamp %.6f) skipped"), TaskId, Timestamp);
+			UE_LOG(LogTraceServices, Log, TEXT("TaskDestroyed(TaskId %" UINT64_FMT ", Timestamp %.6f) skipped"), TaskId, Timestamp);
 			return;
 		}
 
@@ -495,14 +495,14 @@ namespace TraceServices
 
 	void FTasksProvider::AddRelative(const TCHAR* RelationType, TaskTrace::FId TaskId, TArray<FTaskInfo::FRelationInfo> FTaskInfo::* RelationsPtr, TaskTrace::FId RelativeId, double Timestamp, uint32 ThreadId)
 	{
-		UE_LOG(LogTraceServices, Verbose, TEXT("%s (%d) added to TaskId: %d, Timestamp %.6f)"), RelationType, RelativeId, TaskId, Timestamp);
+		UE_LOG(LogTraceServices, Verbose, TEXT("%s (%" UINT64_FMT ") added to TaskId: %" UINT64_FMT ", Timestamp %.6f)"), RelationType, RelativeId, TaskId, Timestamp);
 
 		InitTaskIdToIndexConversion(TaskId);
 
 		FTaskInfo* Task = TryGetTask(TaskId);
 		if (Task == nullptr)
 		{
-			UE_LOG(LogTraceServices, Log, TEXT("Add%s(TaskId %d, OtherId: %d, Timestamp %.6f) skipped"), RelationType, TaskId, RelativeId, Timestamp);
+			UE_LOG(LogTraceServices, Log, TEXT("Add%s(TaskId %" UINT64_FMT ", OtherId: %" UINT64_FMT ", Timestamp %.6f) skipped"), RelationType, TaskId, RelativeId, Timestamp);
 			return;
 		}
 
diff --git a/Engine/Source/Developer/Zen/Private/ZenSerialization.cpp b/Engine/Source/Developer/Zen/Private/ZenSerialization.cpp
--- a/Engine/Source/Developer/Zen/Private/ZenSerialization.cpp
+++ b/Engine/Source/Developer/Zen/Private/ZenSerialization.cpp
@@ -351,7 +351,7 @@ bool TryLoadCbPackage(FCbPackage& Package, FArchive& Ar, FCbBufferAllocator Allo
 					BOOL OK = GetFileSizeEx(Handle, &FileSize);
 					if (!OK)
 					{
-						UE_LOG(LogZenSerialization, Warning, TEXT("Unable to read local file via handle '%u', treating it as a missing attachment"), HandleNum);
+						UE_LOG(LogZenSerialization, Warning, TEXT("Unable to read local file via handle '%" UINT64_FMT "', treating it as a missing attachment"), HandleNum);
 						++Index;
 						continue;
 					}
@@ -369,7 +369,7 @@ bool TryLoadCbPackage(FCbPackage& Package, FArchive& Ar, FCbBufferAllocator Allo
 						if (!Success)
 						{
 							DWORD LastError = GetLastError();
-							UE_LOG(LogZenSerialization, Warning, TEXT("Unable to read local file via handle '%u', error: %u, treating it as a missing attachment"), HandleNum, LastError);
+							UE_LOG(LogZenSerialization, Warning, TEXT("Unable to read local file via handle '%" UINT64_FMT "', error: %u, treating it as a missing attachment"), HandleNum, LastError);
 							break;
 						}
 						Offset += Length;
