Improve patching and paking performance.
Make parts of the paking process parallel to perform a patch.

Parallelize which configs should be staged.

diff --git a/Engine/Source/Developer/PakFileUtilities/Private/PakFileUtilities.cpp b/Engine/Source/Developer/PakFileUtilities/Private/PakFileUtilities.cpp
--- a/Engine/Source/Developer/PakFileUtilities/Private/PakFileUtilities.cpp
+++ b/Engine/Source/Developer/PakFileUtilities/Private/PakFileUtilities.cpp
@@ -4960,90 +4960,101 @@ void RemoveIdenticalFiles( TArray<FPakInputPair>& FilesToPak, const FString& Sou
 	{
 		UAssetToUexpMapping.Add(-1);
 	}
-	for (int i = 0; i<FilesToPak.Num(); i++)
-	{
-		const auto& NewFile = FilesToPak[i];
-		FString Ext(FPaths::GetExtension(FilesToPak[i].Source));
-		if (Ext.Equals("uasset", ESearchCase::IgnoreCase) || Ext.Equals("umap", ESearchCase::IgnoreCase))
+
+	ParallelFor(FilesToPak.Num(), [&FilesToPak, &SourceFileToIndex, &UAssetToUexpMapping](int32 Idx)
 		{
-			FString UexpDestFilename = FPaths::ChangeExtension(NewFile.Source, "uexp");
-			int32 *UexpIndexPtr = SourceFileToIndex.Find(UexpDestFilename);
-			if (UexpIndexPtr)
+			const auto& NewFile = FilesToPak[Idx];
+			FString Ext(FPaths::GetExtension(FilesToPak[Idx].Source));
+			if (Ext.Equals("uasset", ESearchCase::IgnoreCase) || Ext.Equals("umap", ESearchCase::IgnoreCase))
 			{
-				UAssetToUexpMapping[*UexpIndexPtr] = i;
-				UAssetToUexpMapping[i] = *UexpIndexPtr;
+				FString UexpDestFilename = FPaths::ChangeExtension(NewFile.Source, "uexp");
+				int32* UexpIndexPtr = SourceFileToIndex.Find(UexpDestFilename);
+				if (UexpIndexPtr)
+				{
+					UAssetToUexpMapping[*UexpIndexPtr] = Idx;
+					UAssetToUexpMapping[Idx] = *UexpIndexPtr;
+				}
 			}
-		}
-	}
+		});
 	TArray<int64> FileSizes;
 	FileSizes.AddDefaulted(FilesToPak.Num());
 
-
-
+	TArray<FString> ChangedFilesArchiveArray;
+	ChangedFilesArchiveArray.SetNum(FilesToPak.Num());
     // Mark files to remove if they're unchanged
-	for (int i= 0; i<FilesToPak.Num(); i++)
-	{
-		const auto& NewFile = FilesToPak[i];
-		if( NewFile.bIsDeleteRecord )
-		{
-			continue;
-		}
-		FString SourceFileNoMountPoint =  NewFile.Dest.Replace(TEXT("../../../"), TEXT(""));
-		FString SourceFilename = SourceDirectory / SourceFileNoMountPoint;
-		
-		const FFileInfo* FoundFileHash = FileHashes.Find(SourceFileNoMountPoint);
-		if (!FoundFileHash)
+	ParallelFor(FilesToPak.Num(), [SourceDirectory, ChangedFilesArchive, &FilesToPak, &FileHashes, &FileSizes, &IncludeFilesMask, &ChangedFilesArchiveArray](int32 Idx)
 		{
-			FoundFileHash = FileHashes.Find(NewFile.Dest);
-		}
-		
- 		if ( !FoundFileHash )
- 		{
- 			UE_LOG(LogPakFile, Display, TEXT("Didn't find hash for %s No mount %s"), *SourceFilename, *SourceFileNoMountPoint);
- 		}
- 
+			const auto& NewFile = FilesToPak[Idx];
+			if (NewFile.bIsDeleteRecord)
+			{
+				return;
+			}
+			FString SourceFileNoMountPoint = NewFile.Dest.Replace(TEXT("../../../"), TEXT(""));
+			FString SourceFilename = SourceDirectory / SourceFileNoMountPoint;
+
+			const FFileInfo* FoundFileHash = FileHashes.Find(SourceFileNoMountPoint);
+			if (!FoundFileHash)
+			{
+				FoundFileHash = FileHashes.Find(NewFile.Dest);
+			}
+
+			if (!FoundFileHash)
+			{
+				UE_LOG(LogPakFile, Display, TEXT("Didn't find hash for %s No mount %s"), *SourceFilename, *SourceFileNoMountPoint);
+			}
+
 #if GUARANTEE_UASSET_AND_UEXP_IN_SAME_PAK
-		// uexp files are always handled with their corresponding uasset file
-		bool bForceInclude = false;
-		if (!FPaths::GetExtension(SourceFilename).Equals("uexp", ESearchCase::IgnoreCase))
-		{
-			bForceInclude = FoundFileHash && FoundFileHash->bForceInclude;
-		}
+			// uexp files are always handled with their corresponding uasset file
+			bool bForceInclude = false;
+			if (!FPaths::GetExtension(SourceFilename).Equals("uexp", ESearchCase::IgnoreCase))
+			{
+				bForceInclude = FoundFileHash && FoundFileHash->bForceInclude;
+			}
 #else
-		bool bForceInclude = FoundFileHash && FoundFileHash->bForceInclude;
+			bool bForceInclude = FoundFileHash && FoundFileHash->bForceInclude;
 #endif
 
-		FString DestFilename = NewFile.Source;
-		bool bFileIsIdentical = FileIsIdentical(SourceFilename, DestFilename, FoundFileHash, &FileSizes[i]);
+			FString DestFilename = NewFile.Source;
+			bool bFileIsIdentical = FileIsIdentical(SourceFilename, DestFilename, FoundFileHash, &FileSizes[Idx]);
 
-		if (ChangedFilesArchive && !bFileIsIdentical)
-		{
-			ChangedFilesArchive->Logf(TEXT("%s\n"),*DestFilename);
-		}
+			if (ChangedFilesArchive && !bFileIsIdentical)
+			{
+				ChangedFilesArchiveArray[Idx] = DestFilename;
+			}
 
-		if (!bForceInclude && bFileIsIdentical)
+			if (!bForceInclude && bFileIsIdentical)
+			{
+				UE_LOG(LogPakFile, Display, TEXT("Source file %s matches dest file %s and will not be included in patch"), *SourceFilename, *DestFilename);
+				// remove from the files to pak list
+				IncludeFilesMask[Idx] = false;
+			}
+		});
+	if (ChangedFilesArchive)
+	{
+		for (const FString& Log : ChangedFilesArchiveArray)
 		{
-			UE_LOG(LogPakFile, Display, TEXT("Source file %s matches dest file %s and will not be included in patch"), *SourceFilename, *DestFilename);
-			// remove from the files to pak list
-			IncludeFilesMask[i] = false;
+			if (!Log.IsEmpty())
+			{
+				ChangedFilesArchive->Logf(TEXT("%s\n"), *Log);
+			}
 		}
 	}
 
 	// Add corresponding UExp/UBulk files to the patch if one is included but not the other (uassets and uexp files must be in the same pak)
 #if GUARANTEE_UASSET_AND_UEXP_IN_SAME_PAK
-	for (int i = 0; i < FilesToPak.Num(); i++)
-	{
-		int32 CounterpartFileIndex= UAssetToUexpMapping[i];
-		if (CounterpartFileIndex != -1)
+	ParallelFor(FilesToPak.Num(), [&FilesToPak, &UAssetToUexpMapping, &IncludeFilesMask](int32 Idx)
 		{
-			if (IncludeFilesMask[i] != IncludeFilesMask[CounterpartFileIndex])
+			int32 CounterpartFileIndex = UAssetToUexpMapping[Idx];
+			if (CounterpartFileIndex != -1)
 			{
-				UE_LOG(LogPakFile, Display, TEXT("One of %s and %s is different from source, so both will be included in patch"), *FilesToPak[i].Source, *FilesToPak[CounterpartFileIndex].Source);
-				IncludeFilesMask[i] = true;
-				IncludeFilesMask[CounterpartFileIndex] = true;
+				if (IncludeFilesMask[Idx] != IncludeFilesMask[CounterpartFileIndex])
+				{
+					UE_LOG(LogPakFile, Display, TEXT("One of %s and %s is different from source, so both will be included in patch"), *FilesToPak[Idx].Source, *FilesToPak[CounterpartFileIndex].Source);
+					IncludeFilesMask[Idx] = true;
+					IncludeFilesMask[CounterpartFileIndex] = true;
+				}
 			}
-		}
-	}
+		});
 #endif
 
 	if (SeekOptParams.Mode != ESeekOptMode::None)
@@ -5056,10 +5067,10 @@ void RemoveIdenticalFiles( TArray<FPakInputPair>& FilesToPak, const FString& Sou
 	for ( int ReadIndex=0; ReadIndex<IncludeFilesMask.Num(); ReadIndex++)
 	{
 		if (IncludeFilesMask[ReadIndex])
-	{
+		{
 			FilesToPak[WriteIndex++] = FilesToPak[ReadIndex];
+		}
 	}
-}
 	int NumToRemove = FilesToPak.Num() - WriteIndex;
 	FilesToPak.RemoveAt(WriteIndex, NumToRemove, EAllowShrinking::Yes);
 
diff --git a/Engine/Source/Programs/AutomationTool/Scripts/CopyBuildToStagingDirectory.Automation.cs b/Engine/Source/Programs/AutomationTool/Scripts/CopyBuildToStagingDirectory.Automation.cs
--- a/Engine/Source/Programs/AutomationTool/Scripts/CopyBuildToStagingDirectory.Automation.cs
+++ b/Engine/Source/Programs/AutomationTool/Scripts/CopyBuildToStagingDirectory.Automation.cs
@@ -2158,7 +2158,8 @@ namespace AutomationScripts
 		static void StageConfigFiles(DeploymentContext SC, DirectoryReference ConfigDir, string PlatformExtensionName)
 		{
 			List<FileReference> ConfigFiles = SC.FindFilesToStage(ConfigDir, "*.ini", StageFilesSearch.AllDirectories);
-			foreach (FileReference ConfigFile in ConfigFiles)
+			ConcurrentBag<FileReference> ConfigsToStage = new ConcurrentBag<FileReference>();
+			ConfigFiles.AsParallel().ForAll(ConfigFile =>
 			{
 				Nullable<bool> ShouldStage = ShouldStageConfigFile(SC, ConfigDir, ConfigFile, PlatformExtensionName);
 				if (ShouldStage == null)
@@ -2170,13 +2171,17 @@ namespace AutomationScripts
 
 				if (ShouldStage ?? true)
 				{
-					SC.StageFile(StagedFileType.UFS, ConfigFile);
-					Logger.LogInformation("Including config file {ConfigFile}", ConfigFile);
+					ConfigsToStage.Add(ConfigFile);
 				}
 				else
 				{
 					Logger.LogInformation("Excluding config file {ConfigFile}", ConfigFile);
 				}
+			});
+			foreach (FileReference ConfigFile in ConfigsToStage)
+			{
+				SC.StageFile(StagedFileType.UFS, ConfigFile);
+				Logger.LogInformation("Including config file {ConfigFile}", ConfigFile);
 			}
 		}
 
