UE: Allow wildcards when extracting a pak (and a few other operations)

diff --git a/Engine/Source/Developer/PakFileUtilities/Private/PakFileUtilities.cpp b/Engine/Source/Developer/PakFileUtilities/Private/PakFileUtilities.cpp
--- a/Engine/Source/Developer/PakFileUtilities/Private/PakFileUtilities.cpp
+++ b/Engine/Source/Developer/PakFileUtilities/Private/PakFileUtilities.cpp
@@ -5269,17 +5269,34 @@ TArray<FPakInputPair> GetNewDeleteRecords( const TArray<FPakInputPair>& InFilesT
 	return DeleteRecords;
 }
 
+bool AtLeastOnePakFileExists(const FString& PakFilename)
+{
+	if (IFileManager::Get().FileExists(*PakFilename))
+	{
+		return true;
+	}
+	else
+	{
+		// If file doesn't exist try using it as a search string, it may contain wild cards
+		TArray<FString> PakFileList;
+		FString PakFileDirectory = FPaths::GetPath(PakFilename);
+		IFileManager::Get().FindFiles(PakFileList, *PakFileDirectory, *FPaths::GetCleanFilename(PakFilename));
+
+		return PakFileList.Num() > 0;
+	}
+}
+
 FString GetPakPath(const TCHAR* SpecifiedPath, bool bIsForCreation)
 {
 	FString PakFilename(SpecifiedPath);
 	FPaths::MakeStandardFilename(PakFilename);
 	
 	// if we are trying to open (not create) it, but BaseDir relative doesn't exist, look in LaunchDir
-	if (!bIsForCreation && !FPaths::FileExists(PakFilename))
+	if (!bIsForCreation && !AtLeastOnePakFileExists(PakFilename))
 	{
 		PakFilename = FPaths::LaunchDir() + SpecifiedPath;
 
-		if (!FPaths::FileExists(PakFilename))
+		if (!AtLeastOnePakFileExists(PakFilename))
 		{
 			UE_LOG(LogPakFile, Fatal, TEXT("Existing pak file %s could not be found (checked against binary and launch directories)"), SpecifiedPath);
 			return TEXT("");
