Add error for when we attempt to create a pak file that already exists at the location

diff --git a/Engine/Source/Developer/PakFileUtilities/Private/PakFileUtilities.cpp b/Engine/Source/Developer/PakFileUtilities/Private/PakFileUtilities.cpp
--- a/Engine/Source/Developer/PakFileUtilities/Private/PakFileUtilities.cpp
+++ b/Engine/Source/Developer/PakFileUtilities/Private/PakFileUtilities.cpp
@@ -1964,6 +1964,12 @@ void LoadKeyChain(const TCHAR* CmdLine, FKeyChain& OutCryptoSettings)
  */
 FArchive* CreatePakWriter(const TCHAR* Filename, const FKeyChain& InKeyChain, bool bSign)
 {
+	if (IFileManager::Get().FileExists(Filename))
+	{
+		UE_LOG(LogPakFile, Error, TEXT("Failed to create pak at %s. File already exists at that location."), Filename);
+		return nullptr;
+	}
+
 	FArchive* Writer = IFileManager::Get().CreateFileWriter(Filename);
 
 	if (Writer)
