Add support for showing the footer size in pak file listings.
Add support for listing multiple pak files at the same time via specifying a directory or a wildcard.

#rnx

diff --git a/Engine/Source/Developer/PakFileUtilities/Private/PakFileUtilities.cpp b/Engine/Source/Developer/PakFileUtilities/Private/PakFileUtilities.cpp
--- a/Engine/Source/Developer/PakFileUtilities/Private/PakFileUtilities.cpp
+++ b/Engine/Source/Developer/PakFileUtilities/Private/PakFileUtilities.cpp
@@ -74,7 +74,7 @@ int64 GDDCHits = 0;
 int64 GDDCMisses = 0;
 #endif
 
-bool ListFilesInPak(const TCHAR * InPakFilename, int64 SizeFilter, bool bIncludeDeleted, const FString& CSVFilename, bool bExtractToMountPoint, const FKeyChain& InKeyChain, bool bAppendFile);
+bool ListFilesInPak(const TCHAR * InPakFilename, int64 SizeFilter, bool bIncludeDeleted, const FString& CSVFilename, bool bExtractToMountPoint, const FKeyChain& InKeyChain, bool bAppendFile, bool bIncludeFooter);
 
 bool SignPakFile(const FString& InPakFilename, const FRSAKeyHandle InSigningKey)
 {
@@ -3066,7 +3066,8 @@ bool FPakWriterContext::Flush()
 				CsvFilename,
 				bExtractToMountPoint,
 				OutputPakFile->KeyChain,
-				bAppendFile);
+				bAppendFile,
+				false);
 			if (!bPerPakCsvFiles)
 			{
 				bAppendFile = true;
@@ -3270,7 +3271,7 @@ bool TestPakFile(const TCHAR* Filename, bool TestHashes)
 	}
 }
 
-bool ListFilesInPak(const TCHAR * InPakFilename, int64 SizeFilter, bool bIncludeDeleted, const FString& CSVFilename, bool bExtractToMountPoint, const FKeyChain& InKeyChain, bool bAppendFile)
+bool ListFilesInPak(const TCHAR * InPakFilename, int64 SizeFilter, bool bIncludeDeleted, const FString& CSVFilename, bool bExtractToMountPoint, const FKeyChain& InKeyChain, bool bAppendFile, bool bIncludeFooter)
 {
 	TRACE_CPUPROFILER_EVENT_SCOPE(ListFilesInPak);
 	IPlatformFile* LowerLevelPlatformFile = &FPlatformFileManager::Get().GetPlatformFile();
@@ -3349,6 +3350,22 @@ bool ListFilesInPak(const TCHAR * InPakFilename, int64 SizeFilter, bool bInclude
 					bWasCompressed ? TEXT("true") : TEXT("false"),
 					Entry.CompressionMethodIndex);
 			}
+			if (bIncludeFooter)
+			{
+				// If desired, emit a pseudoentry for our metadata. This is all stored at the end of the file
+				// so rather than ask the pak file anything we just assign it to everything past the end of the 
+				// records.
+				int64 RecordEndOffset = 0;
+				if (Records.Num())
+				{
+					const FPakEntry& LastEntry = Records[Records.Num()-1].Value;
+					RecordEndOffset = LastEntry.Offset + LastEntry.Size;
+				}
+				OutputArchive->Logf(
+					TEXT("<Footer>, %lld, %lld, 0, false, false, 0"),
+					RecordEndOffset, PakFile.TotalSize() - RecordEndOffset);
+			}
+
 			OutputArchive->Flush();
 			UE_LOG(LogPakFile, Display, TEXT("Saved CSV file to %s"), *CSVFilename);
 			return true;
@@ -5713,15 +5730,83 @@ bool ExecuteUnrealPak(const TCHAR* CmdLine)
 
 		bool bExcludeDeleted = FParse::Param( CmdLine, TEXT("ExcludeDeleted") );
 
-		FString PakFilename = GetPakPath(*NonOptionArguments[0], false);
+		// Handle if we're trying to list an entire directory or anything with wildcards.
+		// This gets a little weird because GetPakPath does a bunch of normalization, and this
+		// is trying to somewhat match behavior in IostoreUtilities.cpp.
+		FString RawFileName = NonOptionArguments[0];
+		TArray<FString> FilesToList;
+		if (IFileManager::Get().DirectoryExists(*RawFileName))
+		{
+			// List the whole directory.
+			FString Directory = RawFileName;
+			FPaths::NormalizeDirectoryName(Directory);
 
+			TArray<FString> FoundPaks;
+			IFileManager::Get().FindFiles(FoundPaks, *(Directory / TEXT("*.pak")), true, false);
+
+			for (const FString& Filename : FoundPaks)
+			{
+				FilesToList.Emplace(Directory / Filename);
+			}
+		}
+		else if (RawFileName.Contains("*"))
+		{
+			// Technically GetPakPath handles wildcards, but it also only returns a single filename (?),
+			// so we now assume if they include a wildcard they want multiple listings.
+			FString Directory = FPaths::GetPath(RawFileName);
+			FPaths::NormalizeDirectoryName(Directory);
+
+			TArray<FString> FoundPaks;
+			IFileManager::Get().FindFiles(FoundPaks, *RawFileName, true, false);
+
+			for (const FString& Filename : FoundPaks)
+			{
+				FilesToList.Emplace(Directory / Filename);
+			}
+		}
+		else
+		{
+			// Otherwise, use the old system so we don't break anything.
+			FilesToList.Emplace(GetPakPath(*NonOptionArguments[0], false));
+		}
+		
 		FString CSVFilename;
 		FParse::Value(CmdLine, TEXT("csv="), CSVFilename);
 
+		// If CSV is a path, then dump multiple CSVs. If it's a single CSV, then append.
+		bool bAppendCSV = false;
+		if (FilesToList.Num() > 1)
+		{
+			bAppendCSV = !IFileManager::Get().DirectoryExists(*CSVFilename);
+		}
+		else
+		{
+			// If it's a single file, then if it's a directory use the pak filename as a base
+			if (IFileManager::Get().DirectoryExists(*CSVFilename))
+			{
+				CSVFilename = CSVFilename / FPaths::GetCleanFilename(FilesToList[0]) + TEXT(".csv");
+			}
+		}
+
 		bool bExtractToMountPoint = FParse::Param(CmdLine, TEXT("ExtractToMountPoint"));
-		bool bAppendFile = false;
+		bool bIncludeFooter = FParse::Param(CmdLine, TEXT("IncludeFooters"));
 
-		return ListFilesInPak(*PakFilename, SizeFilter, !bExcludeDeleted, *CSVFilename, bExtractToMountPoint, KeyChain, bAppendFile);
+		bool bAllSucceeded = true;
+		for (const FString& PakFilename : FilesToList)
+		{
+			FString ThisPaksCSV = CSVFilename;
+			if (!bAppendCSV && FilesToList.Num() > 1)
+			{
+				// It's a directory.
+				ThisPaksCSV = CSVFilename / FPaths::GetCleanFilename(PakFilename) + TEXT(".csv");
+			}
+
+			if (!ListFilesInPak(*PakFilename, SizeFilter, !bExcludeDeleted, *ThisPaksCSV, bExtractToMountPoint, KeyChain, bAppendCSV, bIncludeFooter))
+			{
+				bAllSucceeded = false;
+			}
+		}
+		return bAllSucceeded;
 	}
 
 	if (FParse::Param(CmdLine, TEXT("Info")))
