Fix UnrealPak output when no CSV option is specified (recent changes would assume it had been specified).

#DisabledValidations Preflight, CheckUnacceptableWords

diff --git a/Engine/Source/Developer/PakFileUtilities/Private/PakFileUtilities.cpp b/Engine/Source/Developer/PakFileUtilities/Private/PakFileUtilities.cpp
--- a/Engine/Source/Developer/PakFileUtilities/Private/PakFileUtilities.cpp
+++ b/Engine/Source/Developer/PakFileUtilities/Private/PakFileUtilities.cpp
@@ -5773,18 +5773,23 @@ bool ExecuteUnrealPak(const TCHAR* CmdLine)
 		FString CSVFilename;
 		FParse::Value(CmdLine, TEXT("csv="), CSVFilename);
 
+		const bool bOutputCSV = !CSVFilename.IsEmpty();
+
 		// If CSV is a path, then dump multiple CSVs. If it's a single CSV, then append.
 		bool bAppendCSV = false;
-		if (FilesToList.Num() > 1)
+		if (bOutputCSV)
 		{
-			bAppendCSV = !IFileManager::Get().DirectoryExists(*CSVFilename);
-		}
-		else
-		{
-			// If it's a single file, then if it's a directory use the pak filename as a base
-			if (IFileManager::Get().DirectoryExists(*CSVFilename))
+			if (FilesToList.Num() > 1)
 			{
-				CSVFilename = CSVFilename / FPaths::GetCleanFilename(FilesToList[0]) + TEXT(".csv");
+				bAppendCSV = !IFileManager::Get().DirectoryExists(*CSVFilename);
+			}
+			else
+			{
+				// If it's a single file, then if it's a directory use the pak filename as a base
+				if (IFileManager::Get().DirectoryExists(*CSVFilename))
+				{
+					CSVFilename = CSVFilename / FPaths::GetCleanFilename(FilesToList[0]) + TEXT(".csv");
+				}
 			}
 		}
 
@@ -5795,7 +5800,7 @@ bool ExecuteUnrealPak(const TCHAR* CmdLine)
 		for (const FString& PakFilename : FilesToList)
 		{
 			FString ThisPaksCSV = CSVFilename;
-			if (!bAppendCSV && FilesToList.Num() > 1)
+			if (bOutputCSV && !bAppendCSV && FilesToList.Num() > 1)
 			{
 				// It's a directory.
 				ThisPaksCSV = CSVFilename / FPaths::GetCleanFilename(PakFilename) + TEXT(".csv");
