Fix PVS warnings

#rnx

diff --git a/Engine/Plugins/FX/Niagara/Source/NiagaraEditor/Private/Sequencer/LevelSequence/NiagaraSystemTrackEditor.cpp b/Engine/Plugins/FX/Niagara/Source/NiagaraEditor/Private/Sequencer/LevelSequence/NiagaraSystemTrackEditor.cpp
--- a/Engine/Plugins/FX/Niagara/Source/NiagaraEditor/Private/Sequencer/LevelSequence/NiagaraSystemTrackEditor.cpp
+++ b/Engine/Plugins/FX/Niagara/Source/NiagaraEditor/Private/Sequencer/LevelSequence/NiagaraSystemTrackEditor.cpp
@@ -111,7 +111,7 @@ void FNiagaraSystemTrackEditor::BuildObjectBindingTrackMenu(FMenuBuilder& MenuBu
 
 		FNiagaraEditorModule& NiagaraEditorModule = FModuleManager::GetModuleChecked<FNiagaraEditorModule>("NiagaraEditor");
 
-		UNiagaraComponent* NiagaraComponent = BoundObjects.Num() == 1 ? Cast<UNiagaraComponent>(BoundObjects[0]) : nullptr;
+		UNiagaraComponent* NiagaraComponent = BoundObjects.Num() == 1 ? Cast<UNiagaraComponent>(BoundObjects[0].Get()) : nullptr;
 		for (const FNiagaraVariable& ParameterVariable : ParameterVariables)
 		{
 			if (ParameterVariable.GetType().IsDataInterface() == false &&
diff --git a/Engine/Plugins/FX/Niagara/Source/NiagaraEditor/Private/ViewModels/NiagaraEmitterViewModel.cpp b/Engine/Plugins/FX/Niagara/Source/NiagaraEditor/Private/ViewModels/NiagaraEmitterViewModel.cpp
--- a/Engine/Plugins/FX/Niagara/Source/NiagaraEditor/Private/ViewModels/NiagaraEmitterViewModel.cpp
+++ b/Engine/Plugins/FX/Niagara/Source/NiagaraEditor/Private/ViewModels/NiagaraEmitterViewModel.cpp
@@ -449,7 +449,7 @@ void FNiagaraEmitterViewModel::OnVMCompiled(FVersionedNiagaraEmitter InEmitter)
 
 				for (int32 i = 0; i < Scripts.Num(); i++)
 				{
-					if (Scripts[i] != nullptr && CompileStatuses.Num() > i && CompileErrors.Num() > i && GPUScript->ContainsUsage(Scripts[i]->Usage))
+					if (CompileStatuses.Num() > i && CompileErrors.Num() > i && Scripts[i] != nullptr && GPUScript->ContainsUsage(Scripts[i]->Usage))
 					{
 						const TArray<FString>& NewErrors = ShaderScript->GetCompileErrors();
 						bool bErrors = false;
diff --git a/Engine/Plugins/Runtime/Metasound/Source/MetasoundFrontend/Public/MetasoundParameterPack.h b/Engine/Plugins/Runtime/Metasound/Source/MetasoundFrontend/Public/MetasoundParameterPack.h
--- a/Engine/Plugins/Runtime/Metasound/Source/MetasoundFrontend/Public/MetasoundParameterPack.h
+++ b/Engine/Plugins/Runtime/Metasound/Source/MetasoundFrontend/Public/MetasoundParameterPack.h
@@ -57,7 +57,7 @@ namespace MetasoundParameterPackPrivate
 	{
 		using ThisType = FMetasoundParameterPackItem<T,MaxNumItems>;
 		FMetasoundParameterPackItem(const FName& InName, const FName& InTypeName, const T& InValue)
-			: FMetasoundParameterPackItemBase(InName, InTypeName, SizeOfParam(InValue), offsetof(FMetasoundParameterPackItem<ThisType>, Value))
+			: FMetasoundParameterPackItemBase(InName, InTypeName, SizeOfParam(InValue), offsetof(FMetasoundParameterPackItem<ThisType>, Value)) //-V1050
 			, Value(InValue)
 		{
 			if constexpr (TIsTArray_V<T>)
diff --git a/Engine/Plugins/VirtualProduction/CaptureData/Source/CaptureDataCore/Private/FrameRange.cpp b/Engine/Plugins/VirtualProduction/CaptureData/Source/CaptureDataCore/Private/FrameRange.cpp
--- a/Engine/Plugins/VirtualProduction/CaptureData/Source/CaptureDataCore/Private/FrameRange.cpp
+++ b/Engine/Plugins/VirtualProduction/CaptureData/Source/CaptureDataCore/Private/FrameRange.cpp
@@ -11,12 +11,10 @@ bool FFrameRange::operator==(const FFrameRange& InOther) const
 
 bool FFrameRange::ContainsFrame(int32 InFrame, const TArray<FFrameRange>& InFrameRangeArray)
 {
-	for (int32 Index = 0; Index < InFrameRangeArray.Num(); ++Index)
+	for (const FFrameRange& FrameRange : InFrameRangeArray)
 	{
-		const FFrameRange& FrameRange = InFrameRangeArray[Index];
-
 		if ((FrameRange.StartFrame >= 0 || FrameRange.EndFrame >= 0) &&
-			(FrameRange.StartFrame < 0 || InFrame >= FrameRange.StartFrame) &&
+			(FrameRange.StartFrame < 0 || InFrame >= FrameRange.StartFrame) && //-V590
 			(FrameRange.EndFrame < 0 || InFrame <= FrameRange.EndFrame))
 		{
 			return true;
diff --git a/Engine/Source/Developer/PakFileUtilities/Private/PakFileUtilities.cpp b/Engine/Source/Developer/PakFileUtilities/Private/PakFileUtilities.cpp
--- a/Engine/Source/Developer/PakFileUtilities/Private/PakFileUtilities.cpp
+++ b/Engine/Source/Developer/PakFileUtilities/Private/PakFileUtilities.cpp
@@ -1040,7 +1040,7 @@ bool PrepareCopyFileToPak(const FString& InMountPoint, const FPakInputPair& InFi
 {
 	const int64 FileSize = UncompressedFile.OriginalSize;
 	const int64 PaddedEncryptedFileSize = Align(FileSize, FAES::AESBlockSize); 
-	if (FileSize < 0)
+	if (FileSize <= 0)
 	{
 		return false;
 	}
@@ -4089,7 +4089,7 @@ bool ExtractFilesFromPak(const TCHAR* InPakFilename, TMap<FString, FFileInfo>& I
 			const int64 BufferSize = 8 * 1024 * 1024; // 8MB buffer for extracting
 			void* Buffer = FMemory::Malloc(BufferSize);
 			int64 CompressionBufferSize = 0;
-			uint8* PersistantCompressionBuffer = NULL;
+			uint8* PersistentCompressionBuffer = NULL;
 			int32 ErrorCount = 0;
 			int32 FileCount = 0;
 			int32 ExtractedCount = 0;
@@ -4158,7 +4158,7 @@ bool ExtractFilesFromPak(const TCHAR* InPakFilename, TMap<FString, FFileInfo>& I
 							}
 							else
 							{
-								UncompressCopyFile(*FileHandle, PakReader.GetArchive(), Entry, PersistantCompressionBuffer, CompressionBufferSize, InKeyChain, PakFile);
+								UncompressCopyFile(*FileHandle, PakReader.GetArchive(), Entry, PersistentCompressionBuffer, CompressionBufferSize, InKeyChain, PakFile);
 							}
 							UE_LOG(LogPakFile, Display, TEXT("Extracted \"%s\" to \"%s\" Offset %lld."), *EntryFileName, *DestFilename, Entry.Offset);
 							ExtractedCount++;
@@ -4197,7 +4197,7 @@ bool ExtractFilesFromPak(const TCHAR* InPakFilename, TMap<FString, FFileInfo>& I
 				}
 			}
 			FMemory::Free(Buffer);
-			FMemory::Free(PersistantCompressionBuffer);
+			FMemory::Free(PersistentCompressionBuffer);
 
 			UE_LOG(LogPakFile, Log, TEXT("Finished extracting %d (including %d errors)."), ExtractedCount, ErrorCount);
 		}
@@ -4315,7 +4315,7 @@ bool DiffFilesInPaks(const FString& InPakFilename1, const FString& InPakFilename
 		const int64 BufferSize = 8 * 1024 * 1024; // 8MB buffer for extracting
 		void* Buffer = FMemory::Malloc(BufferSize);
 		int64 CompressionBufferSize = 0;
-		uint8* PersistantCompressionBuffer = NULL;
+		uint8* PersistentCompressionBuffer = NULL;
 		int32 ErrorCount = 0;
 		int32 FileCount = 0;
 
@@ -4383,7 +4383,7 @@ bool DiffFilesInPaks(const FString& InPakFilename1, const FString& InPakFilename
 				}
 				else
 				{
-					UncompressCopyFile(PAKWriter1, PakReader1.GetArchive(), Entry1, PersistantCompressionBuffer, CompressionBufferSize, InKeyChain1, PakFile1);
+					UncompressCopyFile(PAKWriter1, PakReader1.GetArchive(), Entry1, PersistentCompressionBuffer, CompressionBufferSize, InKeyChain1, PakFile1);
 				}
 
 				if (EntryInfo2.CompressionMethodIndex == 0)
@@ -4392,7 +4392,7 @@ bool DiffFilesInPaks(const FString& InPakFilename1, const FString& InPakFilename
 				}
 				else
 				{
-					UncompressCopyFile(PAKWriter2, PakReader2.GetArchive(), Entry2, PersistantCompressionBuffer, CompressionBufferSize, InKeyChain2, PakFile2);
+					UncompressCopyFile(PAKWriter2, PakReader2.GetArchive(), Entry2, PersistentCompressionBuffer, CompressionBufferSize, InKeyChain2, PakFile2);
 				}
 
 				if (FMemory::Memcmp(PAKWriter1.GetData(), PAKWriter2.GetData(), EntryInfo1.UncompressedSize) != 0)
@@ -4510,7 +4510,7 @@ bool GenerateHashesFromPak(const TCHAR* InPakFilename, const TCHAR* InDestPakFil
 			const int64 BufferSize = 8 * 1024 * 1024; // 8MB buffer for extracting
 			void* Buffer = FMemory::Malloc(BufferSize);
 			int64 CompressionBufferSize = 0;
-			uint8* PersistantCompressionBuffer = NULL;
+			uint8* PersistentCompressionBuffer = NULL;
 			int32 ErrorCount = 0;
 			int32 FileCount = 0;
 
@@ -4567,7 +4567,7 @@ bool GenerateHashesFromPak(const TCHAR* InPakFilename, const TCHAR* InDestPakFil
 							}
 							else
 							{
-								UncompressCopyFile(*FileHandle, PakReader.GetArchive(), Entry, PersistantCompressionBuffer, CompressionBufferSize, InKeyChain, PakFile);
+								UncompressCopyFile(*FileHandle, PakReader.GetArchive(), Entry, PersistentCompressionBuffer, CompressionBufferSize, InKeyChain, PakFile);
 							}
 
 							UE_LOG(LogPakFile, Display, TEXT("Generated hash for \"%s\""), *FullFilename);
@@ -4605,7 +4605,7 @@ bool GenerateHashesFromPak(const TCHAR* InPakFilename, const TCHAR* InDestPakFil
 				}
 			}
 			FMemory::Free(Buffer);
-			FMemory::Free(PersistantCompressionBuffer);
+			FMemory::Free(PersistentCompressionBuffer);
 
 			UE_LOG(LogPakFile, Log, TEXT("Finished extracting %d files (including %d errors)."), FileCount, ErrorCount);
 		}
diff --git a/Engine/Source/Developer/ProfileVisualizer/Private/STimeline.cpp b/Engine/Source/Developer/ProfileVisualizer/Private/STimeline.cpp
--- a/Engine/Source/Developer/ProfileVisualizer/Private/STimeline.cpp
+++ b/Engine/Source/Developer/ProfileVisualizer/Private/STimeline.cpp
@@ -70,7 +70,7 @@ int32 STimeline::OnPaint( const FPaintArgs& Args, const FGeometry& AllottedGeome
 	
 	for( int32 LineIndex = 0; LineIndex <= NumValues; LineIndex++ )
 	{
-		const float NormalizedX = (float)LineIndex / NumValues;
+		const float NormalizedX = (float)LineIndex / NumValues; //-V609 
 		const float LineX = Offset + NormalizedX * Zoom;
 		if( LineX < 0.0f || LineX > 1.0f )
 		{
diff --git a/Engine/Source/Runtime/Core/Public/Templates/EqualTo.h b/Engine/Source/Runtime/Core/Public/Templates/EqualTo.h
--- a/Engine/Source/Runtime/Core/Public/Templates/EqualTo.h
+++ b/Engine/Source/Runtime/Core/Public/Templates/EqualTo.h
@@ -25,8 +25,8 @@ struct TEqualTo
 		else
 		{
 			static_assert(sizeof(T) == 0, "Trying to use TEqualTo<T> where T doesn't have an appropriate operator== overload. Please add bool operator==(T, T), do not specialize TEqualTo.");
-		} //-V591
-	}
+		} 
+	} //-V591
 };
 
 template <>
@@ -42,6 +42,6 @@ struct TEqualTo<void>
 		else
 		{
 			static_assert(sizeof(T) == 0, "Trying to use TEqualTo<void> with types without an appropriate operator== overload. Please add bool operator==(T, U), do not specialize TEqualTo.");
-		} //-V591
-	}
+		} 
+	} //-V591
 };
diff --git a/Engine/Source/Runtime/Core/Public/Templates/Greater.h b/Engine/Source/Runtime/Core/Public/Templates/Greater.h
--- a/Engine/Source/Runtime/Core/Public/Templates/Greater.h
+++ b/Engine/Source/Runtime/Core/Public/Templates/Greater.h
@@ -26,8 +26,8 @@ struct TGreater
 		else
 		{
 			static_assert(sizeof(T) == 0, "Trying to use TGreater<T> where T doesn't have an appropriate operator< overload. Please add bool operator<(T, T), do not specialize TGreater.");
-		} //-V591
-	}
+		} 
+	} //-V591
 };
 
 template <>
@@ -43,6 +43,6 @@ struct TGreater<void>
 		else
 		{
 			static_assert(sizeof(T) == 0, "Trying to use TGreater<void> with types without an appropriate operator< overload. Please add bool operator<(T, U), do not specialize TGreater.");
-		} //-V591
-	}
+		} 
+	} //-V591
 };
diff --git a/Engine/Source/Runtime/Core/Public/Templates/Less.h b/Engine/Source/Runtime/Core/Public/Templates/Less.h
--- a/Engine/Source/Runtime/Core/Public/Templates/Less.h
+++ b/Engine/Source/Runtime/Core/Public/Templates/Less.h
@@ -26,8 +26,8 @@ struct TLess
 		else
 		{
 			static_assert(sizeof(T) == 0, "Trying to use TLess<T> where T doesn't have an appropriate operator< overload. Please add bool operator<(T, T), do not specialize TLess.");
-		} //-V591
-	}
+		} 
+	} //-V591
 };
 
 template <>
@@ -43,6 +43,6 @@ struct TLess<void>
 		else
 		{
 			static_assert(sizeof(T) == 0, "Trying to use TLess<void> with types without an appropriate operator< overload. Please add bool operator<(T, U), do not specialize TLess.");
-		} //-V591
-	}
+		} 
+	} //-V591
 };
