UE: Allow zero byte files to be added to the pak. The was behavior before CL#47092348

diff --git a/Engine/Source/Developer/PakFileUtilities/Private/PakFileUtilities.cpp b/Engine/Source/Developer/PakFileUtilities/Private/PakFileUtilities.cpp
--- a/Engine/Source/Developer/PakFileUtilities/Private/PakFileUtilities.cpp
+++ b/Engine/Source/Developer/PakFileUtilities/Private/PakFileUtilities.cpp
@@ -1040,7 +1040,7 @@ bool PrepareCopyFileToPak(const FString& InMountPoint, const FPakInputPair& InFi
 {
 	const int64 FileSize = UncompressedFile.OriginalSize;
 	const int64 PaddedEncryptedFileSize = Align(FileSize, FAES::AESBlockSize); 
-	if (FileSize <= 0)
+	if (FileSize < 0)
 	{
 		return false;
 	}
