Fix divide by zero warning
#rnx

diff --git a/Engine/Source/Developer/PakFileUtilities/Private/PakFileUtilities.cpp b/Engine/Source/Developer/PakFileUtilities/Private/PakFileUtilities.cpp
--- a/Engine/Source/Developer/PakFileUtilities/Private/PakFileUtilities.cpp
+++ b/Engine/Source/Developer/PakFileUtilities/Private/PakFileUtilities.cpp
@@ -1059,12 +1059,14 @@ bool PrepareCopyFileToPak(const FString& InMountPoint, const FPakInputPair& InFi
 		uint8* UncompressedBuffer = UncompressedFile.CompressedBuffer.Get();
 		if (InFile.bNeedEncryption)
 		{
-			checkSlow(FileSize > 0 || PaddedEncryptedFileSize == 0);
-			for(int64 FillIndex = FileSize; FillIndex < PaddedEncryptedFileSize; ++FillIndex)
+			if (FileSize > 0)
 			{
-				// Fill the trailing buffer with bytes from file. Note that this is now from a fixed location
-				// rather than a random one so that we produce deterministic results
-				UncompressedBuffer[FillIndex] = UncompressedBuffer[(FillIndex - FileSize)%FileSize];
+				for (int64 FillIndex = FileSize; FillIndex < PaddedEncryptedFileSize; ++FillIndex)
+				{
+					// Fill the trailing buffer with bytes from file. Note that this is now from a fixed location
+					// rather than a random one so that we produce deterministic results
+					UncompressedBuffer[FillIndex] = UncompressedBuffer[(FillIndex - FileSize) % FileSize];
+				}
 			}
 
 			//Encrypt the buffer before writing it to disk
