make pakfile utils list pak files way faster

diff --git a/Engine/Source/Developer/PakFileUtilities/Private/PakFileUtilities.cpp b/Engine/Source/Developer/PakFileUtilities/Private/PakFileUtilities.cpp
--- a/Engine/Source/Developer/PakFileUtilities/Private/PakFileUtilities.cpp
+++ b/Engine/Source/Developer/PakFileUtilities/Private/PakFileUtilities.cpp
@@ -3317,9 +3317,19 @@ bool ListFilesInPak(const TCHAR * InPakFilename, int64 SizeFilter, bool bInclude
 		HashesBuffer.SetNum(FileCount * sizeof(FPakEntry::Hash));
 		uint8* Hashes = HashesBuffer.GetData();
 		int32 EntryIndex = 0;
-		for (const TPair<TOptional<FString>, FPakEntry>& It : Records)
 		{
-			PakFile.ReadHashFromPayload(It.Value, Hashes + (EntryIndex++)*sizeof(FPakEntry::Hash));
+			TRACE_CPUPROFILER_EVENT_SCOPE(ReadHashFromPayload);
+
+			FSharedPakReader SharedReader = PakFile.GetSharedReader(nullptr);
+
+			for (const TPair<TOptional<FString>, FPakEntry>& It : Records)
+			{
+				SharedReader->Seek(It.Value.Offset);
+
+				FPakEntry SerializedPakEntry;
+				SerializedPakEntry.Serialize(SharedReader.GetArchive(), PakFile.GetInfo().Version);
+				FMemory::Memcpy(Hashes + (EntryIndex++) * sizeof(FPakEntry::Hash), &SerializedPakEntry.Hash, sizeof(SerializedPakEntry.Hash));
+			}
 		}
 
 		if (CSVFilename.Len() > 0)
