Make ListAllFilesInPak run in parallel for another nice speed boost, roughly 2x most recent update and ~20x original code

diff --git a/Engine/Source/Developer/PakFileUtilities/Private/PakFileUtilities.cpp b/Engine/Source/Developer/PakFileUtilities/Private/PakFileUtilities.cpp
--- a/Engine/Source/Developer/PakFileUtilities/Private/PakFileUtilities.cpp
+++ b/Engine/Source/Developer/PakFileUtilities/Private/PakFileUtilities.cpp
@@ -5810,8 +5810,21 @@ bool ExecuteUnrealPak(const TCHAR* CmdLine)
 		bool bIncludeFooter = FParse::Param(CmdLine, TEXT("IncludeFooters"));
 
 		bool bAllSucceeded = true;
-		for (const FString& PakFilename : FilesToList)
-		{
+		TArray<bool> bSucceeded;
+		bSucceeded.SetNum(FilesToList.Num());
+
+		ParallelFor(FilesToList.Num(), [&FilesToList, 
+										&CSVFilename, 
+										&bOutputCSV, 
+										&bAppendCSV, 
+										&bExtractToMountPoint, 
+										&KeyChain, 
+										&bIncludeFooter, 
+										&SizeFilter,
+										&bExcludeDeleted,
+										&bSucceeded](int32 Idx)
+		{
+			const FString& PakFilename = FilesToList[Idx];
 			FString ThisPaksCSV = CSVFilename;
 			if (bOutputCSV && !bAppendCSV && FilesToList.Num() > 1)
 			{
@@ -5819,11 +5832,14 @@ bool ExecuteUnrealPak(const TCHAR* CmdLine)
 				ThisPaksCSV = CSVFilename / FPaths::GetCleanFilename(PakFilename) + TEXT(".csv");
 			}
 
-			if (!ListFilesInPak(*PakFilename, SizeFilter, !bExcludeDeleted, *ThisPaksCSV, bExtractToMountPoint, KeyChain, bAppendCSV, bIncludeFooter))
-			{
-				bAllSucceeded = false;
-			}
+			bSucceeded[Idx] = ListFilesInPak(*PakFilename, SizeFilter, !bExcludeDeleted, *ThisPaksCSV, bExtractToMountPoint, KeyChain, bAppendCSV, bIncludeFooter);
+		});
+
+		for (const bool &bDidSucceed : bSucceeded)
+		{
+			bAllSucceeded &= bDidSucceed;
 		}
+
 		return bAllSucceeded;
 	}
 
