- Remove the IniKeyDenylist and IniSectionDenylist from all layers of the BinaryConfig, instead of just the final InMemoryFile

diff --git a/Engine/Source/Developer/PakFileUtilities/Private/PakFileUtilities.cpp b/Engine/Source/Developer/PakFileUtilities/Private/PakFileUtilities.cpp
--- a/Engine/Source/Developer/PakFileUtilities/Private/PakFileUtilities.cpp
+++ b/Engine/Source/Developer/PakFileUtilities/Private/PakFileUtilities.cpp
@@ -5574,21 +5574,40 @@ bool MakeBinaryConfig(const TCHAR* CmdLine)
 
 	for (const FString& Filename : Config.GetFilenames())
 	{
-		FConfigFile* File = Config.FindConfigFile(Filename);
+		FConfigBranch* Branch = Config.FindBranch(*Filename, Filename);
+		Branch->RunOnEachFile([&SectionsDenyList, &KeysDenyList](FConfigFile& File, const FString& Name)
+			{
+				for (const FString& Section : SectionsDenyList)
+				{
+					File.Remove(Section);
+				}
 
-		for (const FString& Section : SectionsDenyList)
-		{
-			File->Remove(Section);
-		}
+				// now go over any remaining sections and remove keys
+				for (const TPair<FString, FConfigSection>& SectionPair : (const FConfigFile&)File)
+				{
+					for (FName Key : KeysDenyList)
+					{
+						File.RemoveKeyFromSection(*SectionPair.Key, Key);
+					}
+				}
+			});
 
-		// now go over any remaining sections and remove keys
-		for (const TPair<FString, FConfigSection>& SectionPair : *(const FConfigFile*)File)
-		{
-			for (FName Key : KeysDenyList)
+		Branch->RunOnEachCommandStream([&SectionsDenyList, &KeysDenyList](FConfigCommandStream& Stream, const FString& Name)
 			{
-				File->RemoveKeyFromSection(*SectionPair.Key, Key);
-			}
-		}
+				for (const FString& Section : SectionsDenyList)
+				{
+					Stream.Remove(Section);
+				}
+
+				// now go over any remaining sections and remove keys
+				for (TPair<FString, FConfigCommandStreamSection>& SectionPair : Stream)
+				{
+					for (FName Key : KeysDenyList)
+					{
+						SectionPair.Value.Remove(Key);
+					}
+				}
+			});
 	}
 
 	// check the deny list removed itself
