UnrealPak: Fix unsafe path traversal outside the extract dir in the extract command.

#rnx

#Synced-CL 51536473

diff --git a/Engine/Source/Developer/PakFileUtilities/Private/PakFileUtilities.cpp b/Engine/Source/Developer/PakFileUtilities/Private/PakFileUtilities.cpp
--- a/Engine/Source/Developer/PakFileUtilities/Private/PakFileUtilities.cpp
+++ b/Engine/Source/Developer/PakFileUtilities/Private/PakFileUtilities.cpp
@@ -19,6 +19,7 @@
 #include "Misc/ConfigCacheIni.h"
 #include "Misc/ConfigContext.h"
 #include "HAL/PlatformFileManager.h"
+#include "Hash/xxhash.h"
 #include "Async/ParallelFor.h"
 #include "Async/AsyncWork.h"
 #include "Modules/ModuleManager.h"
@@ -4087,6 +4088,7 @@ bool ExtractFilesFromPak(const TCHAR* InPakFilename, TMap<FString, FFileInfo>& I
 
 	bool bIncludeDeleted = (OutDeletedEntries != nullptr);
 
+	const FString DestPath = FPaths::ConvertRelativePathToFull(InDestPath);
 	for (int32 PakFileIndex = 0; PakFileIndex < PakFileList.Num(); PakFileIndex++)
 	{
 		FString PakFilename = PakFileDirectory + "\\" + PakFileList[PakFileIndex];
@@ -4117,7 +4119,6 @@ bool ExtractFilesFromPak(const TCHAR* InPakFilename, TMap<FString, FFileInfo>& I
 				return false;
 			}
 
-			FString DestPath(InDestPath);
 			FSharedPakReader PakReader = PakFile.GetSharedReader(NULL);
 			const int64 BufferSize = 8 * 1024 * 1024; // 8MB buffer for extracting
 			void* Buffer = FMemory::Malloc(BufferSize);
@@ -4134,7 +4135,7 @@ bool ExtractFilesFromPak(const TCHAR* InPakFilename, TMap<FString, FFileInfo>& I
 				FFileInfo* HashFileInfo = InFileHashes.Find(EntryFileName);
 				if (HashFileInfo == nullptr || HashFileInfo->PatchIndex == PakPriority)
 				{
-					FString DestFilename(DestPath / PakMountPoint / EntryFileName);
+					FString DestFilename = FPaths::ConvertRelativePathToFull(DestPath / PakMountPoint / EntryFileName);
 
 					const FPakEntry& Entry = It.Info();
 					if (Entry.IsDeleteRecord())
@@ -4155,6 +4156,17 @@ bool ExtractFilesFromPak(const TCHAR* InPakFilename, TMap<FString, FFileInfo>& I
 						continue;
 					}
 
+					if (!FPaths::IsUnderDirectory(DestFilename, DestPath))
+					{
+						FString EntryPath = FPaths::GetPath(EntryFileName);
+						uint64 EntryHash = FXxHash64::HashBuffer(*EntryPath, EntryPath.Len() * sizeof(TCHAR)).Hash;
+						FString HashDir = FString::Printf(TEXT("%016llx"), EntryHash);
+						FString SafeDestFilename = FPaths::ConvertRelativePathToFull(DestPath / HashDir / FPaths::GetCleanFilename(EntryFileName));
+						UE_LOG(LogPakFile, Display, TEXT("Fixup extract path for \"%s\": from \"%s\" to \"%s\""),
+							*EntryFileName, *DestFilename, *SafeDestFilename);
+						DestFilename = SafeDestFilename;
+					}
+
 					PakReader->Seek(Entry.Offset);
 					uint32 SerializedCrcTest = 0;
 					FPakEntry EntryInfo;
